Bink: Fix for lint issues and non-unity build issues regarding pure virtual functions and missing header files.

diff --git a/Engine/Plugins/Media/BinkMedia/Source/BinkMediaPlayer/Public/BinkMediaTexture.h b/Engine/Plugins/Media/BinkMedia/Source/BinkMediaPlayer/Public/BinkMediaTexture.h
--- a/Engine/Plugins/Media/BinkMedia/Source/BinkMediaPlayer/Public/BinkMediaTexture.h
+++ b/Engine/Plugins/Media/BinkMedia/Source/BinkMediaPlayer/Public/BinkMediaTexture.h
@@ -36,6 +36,8 @@ class BINKMEDIAPLAYER_API UBinkMediaTexture : public UTexture
 	virtual EMaterialValueType GetMaterialType() const override { return MCT_Texture2D; }
 	virtual float GetSurfaceWidth() const override { return CachedDimensions.X; }
 	virtual float GetSurfaceHeight() const override { return CachedDimensions.Y; }
+	virtual float GetSurfaceDepth() const override { return 0; }
+	virtual uint32 GetSurfaceArraySize() const override { return 0; }
 
 	// UObject overrides.
 
diff --git a/Engine/Plugins/Media/BinkMedia/Source/BinkMediaPlayerEditor/Private/Factories/BinkMediaTextureFactoryNew.h b/Engine/Plugins/Media/BinkMedia/Source/BinkMediaPlayerEditor/Private/Factories/BinkMediaTextureFactoryNew.h
--- a/Engine/Plugins/Media/BinkMedia/Source/BinkMediaPlayerEditor/Private/Factories/BinkMediaTextureFactoryNew.h
+++ b/Engine/Plugins/Media/BinkMedia/Source/BinkMediaPlayerEditor/Private/Factories/BinkMediaTextureFactoryNew.h
@@ -3,6 +3,7 @@
 
 #pragma once
 
+#include "Factories/Factory.h"
 #include "BinkMediaTextureFactoryNew.generated.h"
 
 UCLASS(hidecategories=Object, MinimalAPI)
