Bink: Fixed header guard incorrect lint error. it think ifndef is ok but if !defined is not. �\_(?)_/�

diff --git a/Engine/Plugins/Media/BinkMedia/Source/SDK/include/egttypes.h b/Engine/Plugins/Media/BinkMedia/Source/SDK/include/egttypes.h
--- a/Engine/Plugins/Media/BinkMedia/Source/SDK/include/egttypes.h
+++ b/Engine/Plugins/Media/BinkMedia/Source/SDK/include/egttypes.h
@@ -6,9 +6,11 @@
 // (C) Copyright 1994-2021 Epic Games Tools, LLC  
 //===============================================
 
-#if !defined(__EGTTYPESH__) && !defined(__RADTYPESH__) && !defined(__RADRR_COREH__) 
-#define __RADTYPESH__
+#ifndef __EGTTYPESH__
 #define __EGTTYPESH__
+
+#if !defined(__RADTYPESH__) && !defined(__RADRR_COREH__) 
+#define __RADTYPESH__
 #define __RADRR_COREH__ // block old rr_core
 
 #define RADCOPYRIGHT "Copyright (C) 1994-2021, Epic Games Tools, LLC"
@@ -618,4 +620,6 @@
 
 #endif // __RADRES__
 
+#endif // __RADTYPESH__, __RADRR_COREH__
+
 #endif // __EGTTYPESH__
