Bink: Merge from Release-4.27

diff --git a/Engine/Plugins/Media/BinkMedia/Source/BinkMediaPlayerEditor/BinkMediaPlayerEditor.Build.cs b/Engine/Plugins/Media/BinkMedia/Source/BinkMediaPlayerEditor/BinkMediaPlayerEditor.Build.cs
--- a/Engine/Plugins/Media/BinkMedia/Source/BinkMediaPlayerEditor/BinkMediaPlayerEditor.Build.cs
+++ b/Engine/Plugins/Media/BinkMedia/Source/BinkMediaPlayerEditor/BinkMediaPlayerEditor.Build.cs
@@ -41,6 +41,7 @@ public class BinkMediaPlayerEditor : ModuleRules
 
         PrivatePCHHeaderFile = "Private/BinkMediaPlayerEditorPCH.h";
 
+		PublicDefinitions.Add("BUILDING_FOR_UNREAL_ONLY=1");
 		PublicDependencyModuleNames.Add("BinkMediaPlayer");
 
 		PrivateIncludePathModuleNames.AddRange(
