Bink: Fix for d3d12 residency issues when VRAM is low.

diff --git a/Engine/Plugins/Media/BinkMedia/Source/BinkMediaPlayer/Private/Assets/BinkMediaPlayer.cpp b/Engine/Plugins/Media/BinkMedia/Source/BinkMediaPlayer/Private/Assets/BinkMediaPlayer.cpp
--- a/Engine/Plugins/Media/BinkMedia/Source/BinkMediaPlayer/Private/Assets/BinkMediaPlayer.cpp
+++ b/Engine/Plugins/Media/BinkMedia/Source/BinkMediaPlayer/Private/Assets/BinkMediaPlayer.cpp
@@ -507,6 +507,7 @@ void UBinkMediaPlayer::UpdateTexture(FRHICommandListImmediate &RHICmdList, FText
 
 	if (bnk && (BinkDrawStyle == 0 || isEditor)) 
 	{
+		RHICmdList.Transition(FRHITransitionInfo(ref, ERHIAccess::Unknown, ERHIAccess::RTV));
 		RHICmdList.EnqueueLambda([bnk=bnk, ref, nativePtr, width, height, ul=BinkDestinationUpperLeft, lr=BinkDestinationLowerRight, tonemap, output_nits, alpha, srgb_decode, is_hdr](FRHICommandListImmediate& RHICmdList) {
 			BinkPluginSetHdrSettings(bnk, tonemap, 1.0f, output_nits);
 			BinkPluginSetAlphaSettings(bnk, alpha);
@@ -523,6 +524,7 @@ void UBinkMediaPlayer::UpdateTexture(FRHICommandListImmediate &RHICmdList, FText
 				BinkPluginScheduleToTexture(bnk, ul.X, ul.Y, lr.X, lr.Y, 0, nativePtr, width, height);
 			}
 		});
+		RHICmdList.Transition(FRHITransitionInfo(ref, ERHIAccess::RTV, ERHIAccess::SRVMask));
 	}
 }
 
