Bink: Remove .bik format from the filepicker (we only support .bk2)

diff --git a/Engine/Plugins/Media/BinkMedia/Source/BinkMediaPlayer/Private/BinkMovieStreamer.cpp b/Engine/Plugins/Media/BinkMedia/Source/BinkMediaPlayer/Private/BinkMovieStreamer.cpp
--- a/Engine/Plugins/Media/BinkMedia/Source/BinkMediaPlayer/Private/BinkMovieStreamer.cpp
+++ b/Engine/Plugins/Media/BinkMedia/Source/BinkMediaPlayer/Private/BinkMovieStreamer.cpp
@@ -217,8 +217,7 @@ bool FBinkMovieStreamer::OpenNextMovie()
 
 	FString MoviePathTbl[] = 
 	{
-		StoredMoviePaths[MovieIndex] + TEXT(".bk2"),
-		StoredMoviePaths[MovieIndex] + TEXT(".bik"),
+		StoredMoviePaths[MovieIndex] + TEXT(".bk2")
 	};
 	for (int i = 0; i < sizeof(MoviePathTbl) / sizeof(MoviePathTbl[0]) && !bnk; ++i) 
 	{
diff --git a/Engine/Plugins/Media/BinkMedia/Source/BinkMediaPlayerEditor/Private/BinkMediaPlayerCustomization.cpp b/Engine/Plugins/Media/BinkMedia/Source/BinkMediaPlayerEditor/Private/BinkMediaPlayerCustomization.cpp
--- a/Engine/Plugins/Media/BinkMedia/Source/BinkMediaPlayerEditor/Private/BinkMediaPlayerCustomization.cpp
+++ b/Engine/Plugins/Media/BinkMedia/Source/BinkMediaPlayerEditor/Private/BinkMediaPlayerCustomization.cpp
@@ -58,7 +58,7 @@ private:
 
 			TSharedPtr<SWindow> ParentWindow = FSlateApplication::Get().FindWidgetWindow(AsShared());
 			void* ParentWindowHandle = (ParentWindow.IsValid() && ParentWindow->GetNativeWindow().IsValid()) ? ParentWindow->GetNativeWindow()->GetOSWindowHandle() : nullptr;
-			FString Filter = TEXT("Bink files (*.bik)|*.bik|Bink 2 files (*.bk2)|*.bk2");
+			FString Filter = TEXT("Bink 2 files (*.bk2)|*.bk2");
 
 			TArray<FString> OutFiles;
 		
diff --git a/Engine/Plugins/Media/BinkMedia/Source/BinkMediaPlayerEditor/Private/BinkMediaPlayerCustomization.h b/Engine/Plugins/Media/BinkMedia/Source/BinkMediaPlayerEditor/Private/BinkMediaPlayerCustomization.h
--- a/Engine/Plugins/Media/BinkMedia/Source/BinkMediaPlayerEditor/Private/BinkMediaPlayerCustomization.h
+++ b/Engine/Plugins/Media/BinkMedia/Source/BinkMediaPlayerEditor/Private/BinkMediaPlayerCustomization.h
@@ -26,7 +26,7 @@ struct FBinkMediaPlayerCustomization : public IDetailCustomization
 	}
 	FString HandleUrlPickerFileTypeFilter() const 
 	{ 
-		return TEXT("Bink files (*.bik)|*.bik|Bink 2 files (*.bk2)|*.bk2"); 
+		return TEXT("Bink 2 files (*.bk2)|*.bk2"); 
 	}
 	void HandleUrlPickerPathPicked( const FString& PickedPath );
 	EVisibility HandleUrlWarningIconVisibility() const;
