Remove unnecessary include directive which causes non-unity build failures for those without XDK support (because PLATFORM_WINGDK is not defined).

diff --git a/Engine/Plugins/Media/BinkMedia/Source/BinkMediaPlayer/Private/BinkMediaPlayerPCH.h b/Engine/Plugins/Media/BinkMedia/Source/BinkMediaPlayer/Private/BinkMediaPlayerPCH.h
--- a/Engine/Plugins/Media/BinkMedia/Source/BinkMediaPlayer/Private/BinkMediaPlayerPCH.h
+++ b/Engine/Plugins/Media/BinkMedia/Source/BinkMediaPlayer/Private/BinkMediaPlayerPCH.h
@@ -2,10 +2,6 @@
 //   Licenced under the Unreal Engine EULA 
 #pragma once
 
-#if defined WINAPI_FAMILY && !defined PLATFORM_WINGDK && !defined _GAMING_XBOX_XBOXONE
-#include "MinWindows.h"
-#endif
-
 #include "binkplugin.h"
 
 #include "HAL/PlatformFileManager.h"
