Bink: Removed Arm32 lib in build.cs, not longer used for UE5

diff --git a/Engine/Plugins/Media/BinkMedia/Source/BinkMediaPlayer/BinkMediaPlayer.Build.cs b/Engine/Plugins/Media/BinkMedia/Source/BinkMediaPlayer/BinkMediaPlayer.Build.cs
--- a/Engine/Plugins/Media/BinkMedia/Source/BinkMediaPlayer/BinkMediaPlayer.Build.cs
+++ b/Engine/Plugins/Media/BinkMedia/Source/BinkMediaPlayer/BinkMediaPlayer.Build.cs
@@ -75,7 +75,6 @@ public class BinkMediaPlayer : ModuleRules
 			else if (Target.Platform == UnrealTargetPlatform.Android)
 			{
 				PublicDependencyModuleNames.Add("Launch");
-				PublicAdditionalLibraries.Add(Path.Combine(LibDirectory, "libBinkUnrealAndroidArm32.a"));
 				PublicAdditionalLibraries.Add(Path.Combine(LibDirectory, "libBinkUnrealAndroidArm64.a"));
 				string PluginPath = Utils.MakePathRelativeTo(ModuleDirectory, Target.RelativeEnginePath);
 				AdditionalPropertiesForReceipt.Add("AndroidPlugin", Path.Combine(PluginPath, "BinkMediaPlayer_APL.xml"));
