Bink: merge from 4.27Plus. Bink now uses RHI backend.

diff --git a/Engine/Plugins/Media/BinkMedia/Source/BinkMediaPlayer/BinkMediaPlayer.Build.cs b/Engine/Plugins/Media/BinkMedia/Source/BinkMediaPlayer/BinkMediaPlayer.Build.cs
--- a/Engine/Plugins/Media/BinkMedia/Source/BinkMediaPlayer/BinkMediaPlayer.Build.cs
+++ b/Engine/Plugins/Media/BinkMedia/Source/BinkMediaPlayer/BinkMediaPlayer.Build.cs
@@ -28,6 +28,7 @@ public class BinkMediaPlayer : ModuleRules
         PublicDependencyModuleNames.Add("RHI");
         PublicDependencyModuleNames.Add("MoviePlayer");
         //PublicDependencyModuleNames.Add("MediaAssets");
+        PublicDependencyModuleNames.Add("Projects");
 
         PrivatePCHHeaderFile = "Private/BinkMediaPlayerPCH.h";
 
@@ -45,6 +46,8 @@ public class BinkMediaPlayer : ModuleRules
         }
 
 		PublicDefinitions.Add("BUILDING_FOR_UNREAL_ONLY=1");
+		PublicDefinitions.Add("__RADNOEXPORTS__=1");
+		PublicDefinitions.Add("__RADINSTATICLIB__=1");
 		RuntimeDependencies.Add("$(ProjectDir)/Content/Movies/..."); // For chunked streaming
 
 		string Lib = LibName;
@@ -75,7 +78,10 @@ public class BinkMediaPlayer : ModuleRules
 			else if (Target.Platform == UnrealTargetPlatform.Android)
 			{
 				PublicDependencyModuleNames.Add("Launch");
-				PublicAdditionalLibraries.Add(Path.Combine(LibDirectory, "libBinkUnrealAndroidArm64.a"));
+				PublicAdditionalLibraries.Add(Path.Combine(LibDirectory, "android", "armv7", "libBinkUnrealAndroid.a"));
+				PublicAdditionalLibraries.Add(Path.Combine(LibDirectory, "android", "arm64", "libBinkUnrealAndroid.a"));
+				PublicAdditionalLibraries.Add(Path.Combine(LibDirectory, "android", "x86", "libBinkUnrealAndroid.a"));
+				PublicAdditionalLibraries.Add(Path.Combine(LibDirectory, "android", "x64", "libBinkUnrealAndroid.a"));
 				string PluginPath = Utils.MakePathRelativeTo(ModuleDirectory, Target.RelativeEnginePath);
 				AdditionalPropertiesForReceipt.Add("AndroidPlugin", Path.Combine(PluginPath, "BinkMediaPlayer_APL.xml"));
 			}
diff --git a/Engine/Plugins/Media/BinkMedia/Source/BinkMediaPlayer/Private/Assets/BinkMediaPlayer.cpp b/Engine/Plugins/Media/BinkMedia/Source/BinkMediaPlayer/Private/Assets/BinkMediaPlayer.cpp
--- a/Engine/Plugins/Media/BinkMedia/Source/BinkMediaPlayer/Private/Assets/BinkMediaPlayer.cpp
+++ b/Engine/Plugins/Media/BinkMedia/Source/BinkMediaPlayer/Private/Assets/BinkMediaPlayer.cpp
@@ -284,7 +284,7 @@ bool UBinkMediaPlayer::Open(const FString& Url)
 	{
 		ENQUEUE_RENDER_COMMAND(BinkMediaPlayer_Open_CloseBink)([bnk=bnk](FRHICommandListImmediate& RHICmdList) 
 		{ 
-			RHICmdList.EnqueueLambda([bnk](FRHICommandListImmediate& RHICmdList) { BinkPluginClose(bnk); });
+			BinkPluginClose(bnk);
 		});
 		bnk = NULL;
 	}
@@ -359,7 +359,7 @@ void UBinkMediaPlayer::Close()
 	{
 		ENQUEUE_RENDER_COMMAND(BinkMediaPlayer_Close_CloseBink)([bnk=bnk](FRHICommandListImmediate& RHICmdList) 
 		{ 
-			RHICmdList.EnqueueLambda([bnk](FRHICommandListImmediate& RHICmdList) { BinkPluginClose(bnk); });
+			BinkPluginClose(bnk);
 		});
 		bnk = NULL;
 	}
@@ -483,58 +483,41 @@ void UBinkMediaPlayer::Tick(float DeltaTime)
 			lry = binkh / screenSize.Y + uly;
 		}
 
-#if !PLATFORM_ANDROID
-		if (bink_gpu_api == BinkGL) 
-		{
-			uly = 1 - uly;
-			lry = 1 - lry;
-		}
+#if PLATFORM_ANDROID
+		uly = 1 - uly;
+		lry = 1 - lry;
 #endif
 
 		ENQUEUE_RENDER_COMMAND(BinkScheduleOverlay)([bnk=bnk,ulx,uly,lrx,lry](FRHICommandListImmediate& RHICmdList) 
 		{ 
-			RHICmdList.EnqueueLambda([bnk,ulx,uly,lrx,lry](FRHICommandListImmediate& RHICmdList) { 
-				BinkPluginSetDrawFlags(bnk, 0);
-				BinkPluginScheduleOverlay(bnk, ulx, uly, lrx, lry, 0);
-			});
-        });
+			BinkPluginSetDrawFlags(bnk, 0);
+			BinkPluginScheduleOverlay(bnk, ulx, uly, lrx, lry, 0);
+		});
 	}
 }
 
-void UBinkMediaPlayer::UpdateTexture(FRHICommandListImmediate &RHICmdList, FTextureRHIRef ref, void *nativePtr, int width, int height, bool isEditor, bool tonemap, int output_nits, float alpha, bool srgb_decode, bool is_hdr) 
+void UBinkMediaPlayer::UpdateTexture(FRHICommandListImmediate &RHICmdList, FTexture2DRHIRef ref, void *nativePtr, int width, int height, bool isEditor, bool tonemap, int output_nits, float alpha, bool srgb_decode, bool is_hdr) 
 {
-    check(IsInRenderingThread());
+	check(IsInRenderingThread());
 
 	if (bnk && (BinkDrawStyle == 0 || isEditor)) 
 	{
-		RHICmdList.Transition(FRHITransitionInfo(ref, ERHIAccess::Unknown, ERHIAccess::RTV));
-		RHICmdList.EnqueueLambda([bnk=bnk, ref, nativePtr, width, height, ul=BinkDestinationUpperLeft, lr=BinkDestinationLowerRight, tonemap, output_nits, alpha, srgb_decode, is_hdr](FRHICommandListImmediate& RHICmdList) {
-			BinkPluginSetHdrSettings(bnk, tonemap, 1.0f, output_nits);
-			BinkPluginSetAlphaSettings(bnk, alpha);
-			BinkPluginSetDrawFlags(bnk, srgb_decode ? BinkDrawDecodeSRGB : 0);
-			BinkPluginSetRenderTargetFormat(bnk, is_hdr ? 1 : 0);
-			if (bink_gpu_api == BinkGL)
-			{
-				uintptr_t gltex = *(int*)nativePtr;
-				//UE_LOG(LogBinkMoviePlayer, Warning, TEXT("gltex = %i"), (int)gltex);
-				BinkPluginScheduleToTexture(bnk, ul.X, ul.Y, lr.X, lr.Y, 0, (void*)gltex, width, height);
-			}
-			else
-			{
-				BinkPluginScheduleToTexture(bnk, ul.X, ul.Y, lr.X, lr.Y, 0, nativePtr, width, height);
-			}
-		});
-		RHICmdList.Transition(FRHITransitionInfo(ref, ERHIAccess::RTV, ERHIAccess::SRVMask));
+		BinkPluginSetHdrSettings(bnk, tonemap, 1.0f, output_nits);
+		BinkPluginSetAlphaSettings(bnk, alpha);
+		BinkPluginSetDrawFlags(bnk, srgb_decode ? BinkDrawDecodeSRGB : 0);
+		BinkPluginSetRenderTargetFormat(bnk, is_hdr ? 1 : 0);
+		BinkPluginScheduleToTexture(bnk, BinkDestinationUpperLeft.X, BinkDestinationUpperLeft.Y, BinkDestinationLowerRight.X, BinkDestinationLowerRight.Y, 0, ref.GetReference(), width, height);
+		BinkActiveTextureRefs.Push(ref);
 	}
 }
 
 void UBinkMediaPlayer::Draw(UTexture *texture, bool tonemap, int out_nits, float alpha, bool srgb_decode, bool hdr) 
 {
-	if (!texture->GetResource()) 
+	if (!bnk || !texture->GetResource()) 
 	{
 		return;
 	}
-	FTextureRHIRef ref = texture->GetResource()->TextureRHI;
+	FTexture2DRHIRef ref = texture->GetResource()->TextureRHI->GetTexture2D();
 	if ((!IsPlaying() && !IsPaused()) || !ref) 
 	{
 		return;
@@ -551,7 +534,7 @@ void UBinkMediaPlayer::Draw(UTexture *texture, bool tonemap, int out_nits, float
 	struct parms_t 
 	{
 		UBinkMediaPlayer *player;
-		FTextureRHIRef ref;
+		FTexture2DRHIRef ref;
 		void *native;
 		int width, height;
 		bool tonemap;
@@ -563,5 +546,5 @@ void UBinkMediaPlayer::Draw(UTexture *texture, bool tonemap, int out_nits, float
 	ENQUEUE_RENDER_COMMAND(BinkMediaPlayer_Draw)([parms](FRHICommandListImmediate& RHICmdList) 
 	{ 
 		parms.player->UpdateTexture(RHICmdList, parms.ref, parms.native, parms.width, parms.height, false, parms.tonemap, parms.out_nits, parms.alpha, parms.srgb_decode, parms.hdr);
-    });
+	});
 }
diff --git a/Engine/Plugins/Media/BinkMedia/Source/BinkMediaPlayer/Private/Assets/BinkMediaTexture.cpp b/Engine/Plugins/Media/BinkMedia/Source/BinkMediaPlayer/Private/Assets/BinkMediaTexture.cpp
--- a/Engine/Plugins/Media/BinkMedia/Source/BinkMediaPlayer/Private/Assets/BinkMediaTexture.cpp
+++ b/Engine/Plugins/Media/BinkMedia/Source/BinkMediaPlayer/Private/Assets/BinkMediaTexture.cpp
@@ -25,7 +25,7 @@ UBinkMediaTexture::UBinkMediaTexture( const FObjectInitializer& ObjectInitialize
 	else 
 	{
 		PixelFormat = PF_B8G8R8A8;
-		DecodeSRGB = false;
+		DecodeSRGB = true;
 	}
 }
 
@@ -124,9 +124,10 @@ void UBinkMediaTexture::InitializeTrack()
 
 	if (MediaPlayer)
 	{
-		CachedDimensions = MediaPlayer->GetDimensions();
-		if (CachedDimensions.X != 0 || CachedDimensions.Y != 0)
+		FIntPoint NewDimensions = MediaPlayer->GetDimensions();
+		if (NewDimensions != CachedDimensions && NewDimensions != FIntPoint(ForceInit))
 		{
+			CachedDimensions = NewDimensions;
 			UpdateResource();
 		}
 	}
diff --git a/Engine/Plugins/Media/BinkMedia/Source/BinkMediaPlayer/Private/Assets/BinkMediaTextureResource.cpp b/Engine/Plugins/Media/BinkMedia/Source/BinkMediaPlayer/Private/Assets/BinkMediaTextureResource.cpp
--- a/Engine/Plugins/Media/BinkMedia/Source/BinkMediaPlayer/Private/Assets/BinkMediaTextureResource.cpp
+++ b/Engine/Plugins/Media/BinkMedia/Source/BinkMediaPlayer/Private/Assets/BinkMediaTextureResource.cpp
@@ -67,18 +67,12 @@ void FBinkMediaTextureResource::InitDynamicRHI()
 
 	FRHICommandListImmediate& RHICmdList = FRHICommandListExecutor::GetImmediateCommandList();
 	{
-		FRHIRenderPassInfo RPInfo(RenderTargetTextureRHI, ERenderTargetActions::Clear_Store);
+		FRHIRenderPassInfo RPInfo(TextureRHI, ERenderTargetActions::Clear_Store);
 		TransitionRenderPassTargets(RHICmdList, RPInfo);
 		RHICmdList.BeginRenderPass(RPInfo, TEXT("ClearTexture"));
 		RHICmdList.EndRenderPass();
 		RHICmdList.SetViewport(0, 0, 0, w, h, 1);
 		RHICmdList.Transition(FRHITransitionInfo(TextureRHI.GetReference(), ERHIAccess::Unknown, ERHIAccess::UAVGraphics));
-
-		// Work-around for UE4 bug when playing a chunk loaded video while also streaming a movie
-		RHICmdList.ImmediateFlush(EImmediateFlushType::FlushRHIThreadFlushResources);
-		RHIFlushResources();
-		RHICmdList.SubmitCommandsHint();
-		FPlatformMisc::MemoryBarrier();
 	}
 }
 
@@ -87,31 +81,25 @@ void FBinkMediaTextureResource::ReleaseDynamicRHI()
 	ReleaseRHI();
 	RenderTargetTextureRHI.SafeRelease();
 	RemoveFromDeferredUpdateList();
-
-	// Work-around for UE4 bug when playing a chunk loaded video while also streaming a movie
-	FRHICommandListImmediate& RHICmdList = FRHICommandListExecutor::GetImmediateCommandList();
-	RHICmdList.ImmediateFlush(EImmediateFlushType::FlushRHIThreadFlushResources);
-	RHIFlushResources();
-	RHICmdList.SubmitCommandsHint();
-	FPlatformMisc::MemoryBarrier();
 }
 
 void FBinkMediaTextureResource::UpdateDeferredResource(FRHICommandListImmediate& RHICmdList, bool bClearRenderTarget) 
 {
+	check(IsInRenderingThread());
 	auto Player = Owner->MediaPlayer;
 	if (!Player || (!Player->IsPlaying() && !Player->IsPaused()) || !TextureRHI) 
 	{
 		return;
 	}
-	FRHITexture2D * tex = TextureRHI->GetTexture2D();
-	if (!tex) 
+	FTexture2DRHIRef tex = TextureRHI->GetTexture2D();
+	if (!tex.GetReference()) 
 	{
 		return;
 	}
 	uint32 width = tex->GetSizeX();
 	uint32 height = tex->GetSizeY();
 	bool is_hdr = PixelFormat != PF_B8G8R8A8;
-	Player->UpdateTexture(RHICmdList, TextureRHI, tex->GetNativeResource(), width, height, false, Owner->Tonemap, Owner->OutputNits, Owner->Alpha, Owner->DecodeSRGB, is_hdr);
+	Player->UpdateTexture(RHICmdList, tex, tex->GetNativeResource(), width, height, false, Owner->Tonemap, Owner->OutputNits, Owner->Alpha, Owner->DecodeSRGB, is_hdr);
 }
 
 void FBinkMediaTextureResource::Clear() 
diff --git a/Engine/Plugins/Media/BinkMedia/Source/BinkMediaPlayer/Private/BinkFunctionLibrary.cpp b/Engine/Plugins/Media/BinkMedia/Source/BinkMediaPlayer/Private/BinkFunctionLibrary.cpp
--- a/Engine/Plugins/Media/BinkMedia/Source/BinkMediaPlayer/Private/BinkFunctionLibrary.cpp
+++ b/Engine/Plugins/Media/BinkMedia/Source/BinkMediaPlayer/Private/BinkFunctionLibrary.cpp
@@ -30,43 +30,16 @@ static void Bink_DrawOverlays_Internal(FRHICommandListImmediate &RHICmdList, con
 		return;
 	}
 
-	RHICmdList.SubmitCommandsHint();
-
-	FRHIRenderPassInfo RPInfo(backbuffer, ERenderTargetActions::Load_Store);
-	TransitionRenderPassTargets(RHICmdList, RPInfo);
-	RHICmdList.BeginRenderPass(RPInfo, TEXT("RenderBink"));
-	FGraphicsPipelineStateInitializer GraphicsPSOInit;
-	RHICmdList.ApplyCachedRenderTargets(GraphicsPSOInit);
-
-	GraphicsPSOInit.RasterizerState = TStaticRasterizerState<FM_Solid, CM_None>::GetRHI();
-	GraphicsPSOInit.BlendState = TStaticBlendState<>::GetRHI();
-	GraphicsPSOInit.DepthStencilState = TStaticDepthStencilState<false, CF_Always>::GetRHI();
-
-	auto* ShaderMap = GetGlobalShaderMap(GMaxRHIFeatureLevel);
-	TShaderMapRef<TOneColorVS<true>> VertexShader(ShaderMap);
-	TShaderMapRef<FOneColorPS> PixelShader(ShaderMap);
-	GraphicsPSOInit.BoundShaderState.VertexDeclarationRHI = GetVertexDeclarationFVector4();
-	GraphicsPSOInit.BoundShaderState.VertexShaderRHI = VertexShader.GetVertexShader();
-	GraphicsPSOInit.BoundShaderState.PixelShaderRHI = PixelShader.GetPixelShader();
-	GraphicsPSOInit.PrimitiveType = PT_TriangleList;
-	SetGraphicsPipelineState(RHICmdList, GraphicsPSOInit, 0);
-	RHICmdList.SetViewport(0, 0, 0, screenSize.X, screenSize.Y, 1);
-
-	RHICmdList.EnqueueLambda([backbuffer,screenSize](FRHICommandListImmediate& RHICmdList) {
-		BINKPLUGINFRAMEINFO FrameInfo = {};
-		FrameInfo.screen_resource = backbuffer->GetNativeResource();
-		FrameInfo.screen_resource_state = 4; // D3D12_RESOURCE_STATE_RENDER_TARGET; (only used in d3d12)
-		FrameInfo.width = screenSize.X;
-		FrameInfo.height = screenSize.Y;
-		FrameInfo.sdr_or_hdr = backbuffer->GetFormat() == PF_A2B10G10R10 ? 1 : 0;
-		FrameInfo.cmdBuf = RHICmdList.GetNativeCommandBuffer();
-		BinkPluginSetPerFrameInfo(&FrameInfo);
-		BinkPluginAllScheduled();
-		BinkPluginDraw(0, 1);
-	});
-	RHICmdList.PostExternalCommandsReset();
-	RHICmdList.EndRenderPass();
-	RHICmdList.SubmitCommandsHint();
+	BINKPLUGINFRAMEINFO FrameInfo = {};
+	FrameInfo.screen_resource = backbuffer.GetReference();
+	FrameInfo.screen_resource_state = 4; // D3D12_RESOURCE_STATE_RENDER_TARGET; (only used in d3d12)
+	FrameInfo.width = screenSize.X;
+	FrameInfo.height = screenSize.Y;
+	FrameInfo.sdr_or_hdr = backbuffer->GetFormat() == PF_A2B10G10R10 ? 1 : 0;
+	FrameInfo.cmdBuf = &RHICmdList;
+	BinkPluginSetPerFrameInfo(&FrameInfo);
+	BinkPluginAllScheduled();
+	BinkPluginDraw(0, 1);
 }
 
 void UBinkFunctionLibrary::Bink_DrawOverlays() 
@@ -80,7 +53,6 @@ void UBinkFunctionLibrary::Bink_DrawOverlays()
 	{ 
 		if (gameViewport.IsValid())
 		{
-			RHICmdList.PostExternalCommandsReset();
 			Bink_DrawOverlays_Internal(RHICmdList, gameViewport.Get());
 		}
 	});
diff --git a/Engine/Plugins/Media/BinkMedia/Source/BinkMediaPlayer/Private/BinkMediaPlayerModule.cpp b/Engine/Plugins/Media/BinkMedia/Source/BinkMediaPlayer/Private/BinkMediaPlayerModule.cpp
--- a/Engine/Plugins/Media/BinkMedia/Source/BinkMediaPlayer/Private/BinkMediaPlayerModule.cpp
+++ b/Engine/Plugins/Media/BinkMedia/Source/BinkMediaPlayer/Private/BinkMediaPlayerModule.cpp
@@ -4,12 +4,16 @@
 #include "BinkMediaPlayerPCH.h"
 #include "BinkFunctionLibrary.h"
 
-DEFINE_LOG_CATEGORY(LogBink);
+#include "Runtime/Core/Public/Misc/Paths.h"
+#include "Interfaces/IPluginManager.h"
+#include "Misc/CoreDelegates.h"
 
-#define LOCTEXT_NAMESPACE "BinkMediaPlayerModule"
+DEFINE_LOG_CATEGORY(LogBink);
 
 TSharedPtr<FBinkMovieStreamer, ESPMode::ThreadSafe> MovieStreamer;
 
+TArray< FTexture2DRHIRef > BinkActiveTextureRefs;
+
 #if BINKPLUGIN_UE4_EDITOR
 class UFactory;
 #include "ISettingsModule.h"
@@ -126,6 +130,9 @@ extern void BinkFreeGpu(void* ptr, UINTa Amt);
 #ifdef BINK_NDA_CPU_ALLOC
 extern void* BinkAllocCpu(UINTa Amt, U32 Align);
 extern void BinkFreeCpu(void* ptr);
+#elif defined BINK_NO_CPU_ALLOC
+#define BinkAllocCpu 0
+#define BinkFreeCpu 0
 #else
 static void *BinkAllocCpu(UINTa Amt, U32 Align) { return FMemory::Malloc(Amt, Align); }
 static void BinkFreeCpu(void * ptr) { FMemory::Free(ptr); }
@@ -147,56 +154,26 @@ struct FBinkMediaPlayerModule : IModuleInterface, FTickableGameObject
 {
 	virtual void StartupModule() override 
 	{
+		if (IsRunningCommandlet())
+		{
+			return;
+		}
+
 		// TODO: make this an INI setting and/or configurable in Project Settings
 		//BinkPluginTurnOnGPUAssist();
 
-		BINKPLUGININITINFO InitInfo = { 0 };
-		InitInfo.queue = GDynamicRHI->RHIGetNativeGraphicsQueue();
-		InitInfo.physical_device = GDynamicRHI->RHIGetNativePhysicalDevice();
+		static BINKPLUGININITINFO InitInfo = { 0 };
+		InitInfo.queue = 0;
+		InitInfo.physical_device = 0;
 		InitInfo.alloc = BinkAllocCpu;
 		InitInfo.free = BinkFreeCpu;
 		InitInfo.gpu_alloc = BinkAllocGpu;
 		InitInfo.gpu_free = BinkFreeGpu;
+		bink_gpu_api = BinkRHI;
 
-		if(IsVulkanPlatform(GMaxRHIShaderPlatform)) 
-		{
-			bink_gpu_api_hdr = IsHDREnabled();
-			bink_gpu_api = BinkVulkan;
-			InitInfo.sdr_and_hdr_render_target_formats[0] = 44; // VK_FORMAT_B8G8R8A8_UNORM;
-			InitInfo.sdr_and_hdr_render_target_formats[1] = 64; // VK_FORMAT_A2B10G10R10_UNORM_PACK32;
-		}
-		else if (FModuleManager::Get().IsModuleLoaded(TEXT("D3D12RHI")) || FModuleManager::Get().IsModuleLoaded(TEXT("XboxOneD3D12RHI")))
-		{
-			bink_gpu_api_hdr = true;
-			bink_gpu_api = BinkD3D12;
-			// NOTE/WARNING: This format should match the backbuffer if you are using overlay rendering! Otherwise you will get D3D Errors!
-			// Note: D3D12 Bink backend only supports a single format currently!
-			InitInfo.sdr_and_hdr_render_target_formats[0] = 24;// DXGI_FORMAT_R10G10B10A2_UNORM;
-			InitInfo.sdr_and_hdr_render_target_formats[1] = 24;// DXGI_FORMAT_R10G10B10A2_UNORM;
-			bink_force_pixel_format = PF_A2B10G10R10;
-		}
-		else if(IsMetalPlatform(GMaxRHIShaderPlatform)) 
-		{
-			bink_gpu_api = BinkMetal;
-            InitInfo.sdr_and_hdr_render_target_formats[0] = 80;// MTLPixelFormatBGRA8Unorm;
-            InitInfo.sdr_and_hdr_render_target_formats[1] = 90;// MTLPixelFormatRGB10A2Unorm;
-		} 
-		else if (IsOpenGLPlatform(GMaxRHIShaderPlatform))
-		{
-			bink_gpu_api = BinkGL;
-		}
-		else if (IsD3DPlatform(GMaxRHIShaderPlatform))
-		{
-			bink_gpu_api = BinkD3D11;
-		}
-		else
-		{
-			bink_gpu_api = BinkNDA;
-		}
-
-		bPluginInitialized = (bool)BinkPluginInit(GDynamicRHI->RHIGetNativeDevice(), &InitInfo, bink_gpu_api);
+		bPluginInitialized = (bool)BinkPluginInit(0, &InitInfo, bink_gpu_api);
 
-		if (!bPluginInitialized) 
+		if (!bPluginInitialized)
 		{
 			printf("Bink Error: %s\n", BinkPluginError());
 		}
@@ -206,18 +183,18 @@ struct FBinkMediaPlayerModule : IModuleInterface, FTickableGameObject
 		GetMoviePlayer()->RegisterMovieStreamer(MovieStreamer);
 
 #if BINKPLUGIN_UE4_EDITOR
-		ISettingsModule* SettingsModule = FModuleManager::GetModulePtr<ISettingsModule>("Settings");
-		if (SettingsModule && !IsRunningGame()) 
+		static ISettingsModule* SettingsModule = FModuleManager::GetModulePtr<ISettingsModule>("Settings");
+		if (SettingsModule && !IsRunningGame())
 		{
 			SettingsModule->RegisterSettings("Project", "Project", "Bink Movies",
 				LOCTEXT("MovieSettingsName", "Bink Movies"),
 				LOCTEXT("MovieSettingsDescription", "Bink Movie player settings"),
 				GetMutableDefault<UBinkMoviePlayerSettings>()
-				);
+			);
 
 			//SettingsModule->RegisterViewer("Project", *this);
 
-			IMainFrameModule& MainFrameModule = FModuleManager::LoadModuleChecked<IMainFrameModule>("MainFrame");
+			static IMainFrameModule& MainFrameModule = FModuleManager::LoadModuleChecked<IMainFrameModule>("MainFrame");
 			MainFrameModule.OnMainFrameCreationFinished().AddRaw(this, &FBinkMediaPlayerModule::Initialize);
 		}
 #endif
@@ -226,10 +203,13 @@ struct FBinkMediaPlayerModule : IModuleInterface, FTickableGameObject
 
 	virtual void ShutdownModule() override 
 	{
-		BinkPluginShutdown();
-		if (overlayHook.IsValid() && GEngine && GEngine->GameViewport) 
+		if (bPluginInitialized)
 		{
-			GEngine->GameViewport->OnDrawn().Remove(overlayHook);
+			BinkPluginShutdown();
+			if (overlayHook.IsValid() && GEngine && GEngine->GameViewport)
+			{
+				GEngine->GameViewport->OnDrawn().Remove(overlayHook);
+			}
 		}
 	}
 
@@ -250,17 +230,15 @@ struct FBinkMediaPlayerModule : IModuleInterface, FTickableGameObject
 	{
 		ENQUEUE_RENDER_COMMAND(BinkProcess)([](FRHICommandListImmediate& RHICmdList) 
 		{ 
-			RHICmdList.SubmitCommandsHint();
-			RHICmdList.EnqueueLambda([](FRHICommandListImmediate& RHICmdList) {
-				BINKPLUGINFRAMEINFO FrameInfo = {};
-				FrameInfo.cmdBuf = RHICmdList.GetNativeCommandBuffer();
-				BinkPluginSetPerFrameInfo(&FrameInfo);
-				BinkPluginProcessBinks(0);
-				BinkPluginAllScheduled();
+			BINKPLUGINFRAMEINFO FrameInfo = {};
+			BinkPluginSetPerFrameInfo(&FrameInfo);
+			BinkPluginProcessBinks(0);
+			BinkPluginAllScheduled();
+			if(BinkActiveTextureRefs.Num())
+			{
 				BinkPluginDraw(1, 0);
-			});
-			RHICmdList.PostExternalCommandsReset();
-			RHICmdList.SubmitCommandsHint();
+				BinkActiveTextureRefs.Empty(256);
+			}
 		});
 		UBinkFunctionLibrary::Bink_DrawOverlays();
 	}
@@ -290,7 +268,5 @@ struct FBinkMediaPlayerModule : IModuleInterface, FTickableGameObject
 	bool bPluginInitialized = false;
 };
 
-#undef LOCTEXT_NAMESPACE
-
 IMPLEMENT_MODULE( FBinkMediaPlayerModule, BinkMediaPlayer )
 
diff --git a/Engine/Plugins/Media/BinkMedia/Source/BinkMediaPlayer/Private/BinkMediaPlayerPCH.h b/Engine/Plugins/Media/BinkMedia/Source/BinkMediaPlayer/Private/BinkMediaPlayerPCH.h
--- a/Engine/Plugins/Media/BinkMedia/Source/BinkMediaPlayer/Private/BinkMediaPlayerPCH.h
+++ b/Engine/Plugins/Media/BinkMedia/Source/BinkMediaPlayer/Private/BinkMediaPlayerPCH.h
@@ -33,6 +33,7 @@ extern BINKMEDIAPLAYER_API unsigned bink_gpu_api;
 extern BINKMEDIAPLAYER_API unsigned bink_gpu_api_hdr;
 extern BINKMEDIAPLAYER_API EPixelFormat bink_force_pixel_format;
 extern BINKMEDIAPLAYER_API FString BinkUE4CookOnTheFlyPath(FString path, const TCHAR *filename);
+extern BINKMEDIAPLAYER_API TArray< FTexture2DRHIRef > BinkActiveTextureRefs;
 
 static int GetNumSpeakers() 
 {
@@ -51,3 +52,9 @@ static int GetNumSpeakers()
 }
 
 DECLARE_LOG_CATEGORY_EXTERN(LogBink, Log, All);
+
+#ifndef LOCTEXT_NAMESPACE
+#define LOCTEXT_NAMESPACE "BinkMediaPlayerModule"
+#endif
+
+
diff --git a/Engine/Plugins/Media/BinkMedia/Source/BinkMediaPlayer/Private/BinkMovieStreamer.cpp b/Engine/Plugins/Media/BinkMedia/Source/BinkMediaPlayer/Private/BinkMovieStreamer.cpp
--- a/Engine/Plugins/Media/BinkMedia/Source/BinkMediaPlayer/Private/BinkMovieStreamer.cpp
+++ b/Engine/Plugins/Media/BinkMedia/Source/BinkMediaPlayer/Private/BinkMovieStreamer.cpp
@@ -127,65 +127,51 @@ bool FBinkMovieStreamer::Tick(float DeltaTime)
 		float lrx = destLowerRight.X;
 		float lry = destLowerRight.Y;
 
-		auto &RHICmdList = GRHICommandList.GetImmediateCommandList();
-
-		RHICmdList.Transition(FRHITransitionInfo(tex, ERHIAccess::Unknown, ERHIAccess::SRVMask));
-		RHICmdList.SubmitCommandsHint();
-		RHICmdList.EnqueueLambda([bnk=bnk, tex, binkw, binkh, ulx, uly, lrx, lry, is_hdr](FRHICommandListImmediate& RHICmdList) {
-			static const auto CVarHDROutputEnabled = IConsoleManager::Get().FindTConsoleVariableDataInt(TEXT("r.HDR.EnableHDROutput"));
-			static const auto CVarDisplayOutputDevice = IConsoleManager::Get().FindTConsoleVariableDataInt(TEXT("r.HDR.Display.OutputDevice"));
-			if (GRHISupportsHDROutput && CVarHDROutputEnabled->GetValueOnRenderThread() != 0)
-			{
-				EDisplayOutputFormat outDev = static_cast<EDisplayOutputFormat>(CVarDisplayOutputDevice->GetValueOnRenderThread());
-				switch (outDev)
-				{
-				// LDR
-				case EDisplayOutputFormat::SDR_sRGB:
-				case EDisplayOutputFormat::SDR_Rec709:
-				case EDisplayOutputFormat::SDR_ExplicitGammaMapping:
-					BinkPluginSetHdrSettings(bnk, 1, 1.0f, 80);
-					break;
-				// 1k nits
-				case EDisplayOutputFormat::HDR_ACES_1000nit_ST2084:
-				case EDisplayOutputFormat::HDR_ACES_1000nit_ScRGB:
-					BinkPluginSetHdrSettings(bnk, 1, 1.0f, 1000);
-					break;
-				// 2k nits
-				case EDisplayOutputFormat::HDR_ACES_2000nit_ST2084:
-				case EDisplayOutputFormat::HDR_ACES_2000nit_ScRGB:
-					BinkPluginSetHdrSettings(bnk, 1, 1.0f, 2000);
-					break;
-				// no tonemap
-				default:
-					BinkPluginSetHdrSettings(bnk, 0, 1.0f, 1000);
-					break;
-				}
-			}
-			else
+		static const auto CVarHDROutputEnabled = IConsoleManager::Get().FindTConsoleVariableDataInt(TEXT("r.HDR.EnableHDROutput"));
+		static const auto CVarDisplayOutputDevice = IConsoleManager::Get().FindTConsoleVariableDataInt(TEXT("r.HDR.Display.OutputDevice"));
+		if (GRHISupportsHDROutput && CVarHDROutputEnabled->GetValueOnRenderThread() != 0)
+		{
+			EDisplayOutputFormat outDev = static_cast<EDisplayOutputFormat>(CVarDisplayOutputDevice->GetValueOnRenderThread());
+			switch (outDev)
 			{
+			// LDR
+			case EDisplayOutputFormat::SDR_sRGB:
+			case EDisplayOutputFormat::SDR_Rec709:
+			case EDisplayOutputFormat::SDR_ExplicitGammaMapping:
 				BinkPluginSetHdrSettings(bnk, 1, 1.0f, 80);
+				break;
+			// 1k nits
+			case EDisplayOutputFormat::HDR_ACES_1000nit_ST2084:
+			case EDisplayOutputFormat::HDR_ACES_1000nit_ScRGB:
+				BinkPluginSetHdrSettings(bnk, 1, 1.0f, 1000);
+				break;
+			// 2k nits
+			case EDisplayOutputFormat::HDR_ACES_2000nit_ST2084:
+			case EDisplayOutputFormat::HDR_ACES_2000nit_ScRGB:
+				BinkPluginSetHdrSettings(bnk, 1, 1.0f, 2000);
+				break;
+			// no tonemap
+			default:
+				BinkPluginSetHdrSettings(bnk, 0, 1.0f, 1000);
+				break;
 			}
-			BinkPluginSetRenderTargetFormat(bnk, is_hdr ? 1 : 0);
-
-			if (bink_gpu_api == BinkGL) 
-			{
-				uintptr_t gltex = *(int*)tex->GetNativeResource();
-				BinkPluginScheduleToTexture(bnk, ulx, uly, lrx, lry, 0, (void*)gltex, binkw, binkh);
-			} 
-			else
-			{
-				BinkPluginScheduleToTexture(bnk, ulx, uly, lrx, lry, 0, tex->GetNativeResource(), binkw, binkh);
-			}
-
-			BINKPLUGINFRAMEINFO FrameInfo = {};
-			FrameInfo.cmdBuf = RHICmdList.GetNativeCommandBuffer();
-			BinkPluginSetPerFrameInfo(&FrameInfo);
-			BinkPluginAllScheduled();
-			BinkPluginProcessBinks(0);
-			BinkPluginDraw(1, 0);
-		});
-		RHICmdList.PostExternalCommandsReset();
-		RHICmdList.SubmitCommandsHint();
+		}
+		else
+		{
+			BinkPluginSetHdrSettings(bnk, 1, 1.0f, 80);
+		}
+		BinkPluginSetRenderTargetFormat(bnk, is_hdr ? 1 : 0);
+
+		BinkPluginScheduleToTexture(bnk, ulx, uly, lrx, lry, 0, tex.GetReference(), binkw, binkh);
+
+		auto& RHICmdList = GRHICommandList.GetImmediateCommandList();
+		
+		BINKPLUGINFRAMEINFO FrameInfo = {};
+		FrameInfo.cmdBuf = &RHICmdList;
+		BinkPluginSetPerFrameInfo(&FrameInfo);
+		BinkPluginAllScheduled();
+		BinkPluginProcessBinks(0);
+		BinkPluginDraw(1, 0);
 
 		MovieViewport->SetTexture(Texture);	
 	}
diff --git a/Engine/Plugins/Media/BinkMedia/Source/BinkMediaPlayer/Private/binkpluginRHI.cpp b/Engine/Plugins/Media/BinkMedia/Source/BinkMediaPlayer/Private/binkpluginRHI.cpp
new file mode 100644
--- /dev/null
+++ b/Engine/Plugins/Media/BinkMedia/Source/BinkMediaPlayer/Private/binkpluginRHI.cpp
@@ -0,0 +1,98 @@
+// Copyright Epic Games, Inc. All Rights Reserved.
+
+#define BINKRHIFUNCTIONS
+#define BINKTEXTURESCLEANUP
+#include "egttypes.h"
+#include "binktiny.h"
+#include "binktextures.h"
+#include "binkplugin.h"
+
+#include "CoreMinimal.h"
+#include "UObject/ConstructorHelpers.h"
+#include "RHI.h"
+#include "Materials/MaterialInstanceDynamic.h"
+#include "Materials/MaterialInterface.h"
+#include "Materials/Material.h"
+#include "Runtime/RHI/Public/RHIStaticStates.h"
+#include "Runtime/RenderCore/Public/ShaderParameterUtils.h"
+#include "Runtime/RenderCore/Public/RenderResource.h"
+#include "Runtime/Renderer/Public/MaterialShader.h"
+#include "Runtime/RenderCore/Public/RenderGraphResources.h"
+#include "Runtime/RenderCore/Public/RenderGraphResources.h"
+
+static BINKSHADERS * shaders;
+
+static int attached = 0;
+static unsigned rfb;
+
+FRHITexture2D *BinkRHIRenderTarget;
+ERenderTargetLoadAction BinkRenderTargetLoadAction;
+
+RADDEFFUNC int setup_rhi( void * device, BINKPLUGININITINFO * info, S32 gpu_assisted, void ** context )
+{
+  if ( shaders == 0 )
+  {
+    shaders = Create_Bink_shaders( device );
+    if ( shaders == 0 )
+      return 0;
+  }
+
+  return 1;
+}
+
+RADDEFFUNC void shutdown_rhi( void )
+{
+  if ( shaders )
+  {
+    Free_Bink_shaders( shaders );
+    shaders = 0;
+  }
+}
+
+RADDEFFUNC void * createtextures_rhi( void * bink )
+{
+  void * ret = 0;
+  if ( shaders )
+  {
+    ret = Create_Bink_textures( shaders, (HBINK)bink, 0 );
+  }
+  return ret;
+}
+
+RADDEFFUNC void selectrendertarget_rhi(void* texture_target, S32 width, S32 height, S32 do_clear, U32 sdr_or_hdr, S32 current_resource_state)
+{
+	if (texture_target)
+	{
+		attached = 1;
+		BinkRHIRenderTarget = (FRHITexture2D*)texture_target;
+		BinkRenderTargetLoadAction = do_clear ? ERenderTargetLoadAction::EClear : ERenderTargetLoadAction::ELoad;
+	}
+}
+
+RADDEFFUNC void selectscreenrendertarget_rhi( void * texture_target, S32 width, S32 height, U32 format_idx, S32 current_resource_state )
+{
+	if (texture_target)
+	{
+		attached = 2;
+		BinkRHIRenderTarget = (FRHITexture2D*)texture_target;
+		BinkRenderTargetLoadAction = ERenderTargetLoadAction::ELoad;
+	}
+}
+
+RADDEFFUNC void clearrendertarget_rhi(void)
+{
+	if (attached)
+	{
+		BinkRHIRenderTarget = NULL;
+		attached = 0;
+	}
+}
+  
+RADDEFFUNC void begindraw_rhi( void )
+{
+}
+
+RADDEFFUNC void enddraw_rhi( void )
+{
+  clearrendertarget_rhi();
+}
diff --git a/Engine/Plugins/Media/BinkMedia/Source/BinkMediaPlayer/Private/binktexturesRHI.cpp b/Engine/Plugins/Media/BinkMedia/Source/BinkMediaPlayer/Private/binktexturesRHI.cpp
new file mode 100644
--- /dev/null
+++ b/Engine/Plugins/Media/BinkMedia/Source/BinkMediaPlayer/Private/binktexturesRHI.cpp
@@ -0,0 +1,594 @@
+// Copyright Epic Games, Inc. All Rights Reserved.
+
+#define BINKRHIFUNCTIONS
+#include "../../SDK/include/egttypes.h"
+#include "../../SDK/include/binktiny.h"
+#include "../../SDK/include/binktextures.h"
+
+#include "CoreMinimal.h"
+#include "UObject/ConstructorHelpers.h"
+#include "RHI.h"
+#include "Materials/MaterialInstanceDynamic.h"
+#include "Materials/MaterialInterface.h"
+#include "Materials/Material.h"
+#include "Runtime/RHI/Public/RHIStaticStates.h"
+#include "Runtime/RHI/Public/PipelineStateCache.h"
+#include "Runtime/RenderCore/Public/ShaderParameterUtils.h"
+#include "Runtime/RenderCore/Public/RenderResource.h"
+#include "Runtime/Renderer/Public/MaterialShader.h"
+#include "Runtime/RenderCore/Public/RenderGraphBuilder.h"
+#include "Runtime/RenderCore/Public/RenderGraphResources.h"
+#include "Runtime/RenderCore/Public/BinkShaders.h"
+
+
+extern FRHITexture2D *BinkRHIRenderTarget;
+extern ERenderTargetLoadAction BinkRenderTargetLoadAction;
+
+FRDGTextureRef BinkRegisterExternalTexture(FRDGBuilder& GraphBuilder, FRHITexture* Texture, const TCHAR* NameIfUnregistered)
+{
+    if (FRDGTextureRef FoundTexture = GraphBuilder.FindExternalTexture(Texture))
+    {
+        return FoundTexture;
+    }
+    return GraphBuilder.RegisterExternalTexture(CreateRenderTarget(Texture, NameIfUnregistered));
+}
+
+typedef struct BINKSHADERSRHI
+{
+	BINKSHADERS pub;
+} BINKSHADERSRHI; 
+
+typedef struct BINKTEXTURESRHI
+{
+	BINKTEXTURES pub;
+
+	// RHI
+	S32 video_width, video_height;
+	S32 dirty;
+
+	BINKSHADERSRHI * shaders;
+	HBINK bink;
+	// this is the Bink info on the textures
+	BINKFRAMEBUFFERS bink_buffers;
+
+	F32 Ax, Ay, Bx, By, Cx, Cy;
+	F32 alpha;
+	S32 draw_type;
+	S32 draw_flags;
+	F32 u0, v0, u1, v1;
+
+	FTexture2DRHIRef Ytexture[BINKMAXFRAMEBUFFERS];
+	FTexture2DRHIRef cRtexture[BINKMAXFRAMEBUFFERS];
+	FTexture2DRHIRef cBtexture[BINKMAXFRAMEBUFFERS];
+	FTexture2DRHIRef Atexture[BINKMAXFRAMEBUFFERS];
+	FTexture2DRHIRef Htexture[BINKMAXFRAMEBUFFERS];
+
+	// unused for now
+	S32 tonemap;
+	F32 exposure;
+	F32 out_luma;
+
+} BINKTEXTURESRHI;
+
+#if defined( BINKTEXTURESINDIRECTBINKCALLS )
+  RADDEFFUNC void indirectBinkGetFrameBuffersInfo( HBINK bink, BINKFRAMEBUFFERS * fbset );
+  #define BinkGetFrameBuffersInfo indirectBinkGetFrameBuffersInfo
+  RADDEFFUNC void indirectBinkRegisterFrameBuffers( HBINK bink, BINKFRAMEBUFFERS * fbset );
+  #define BinkRegisterFrameBuffers indirectBinkRegisterFrameBuffers
+  RADDEFFUNC S32 indirectBinkAllocateFrameBuffers( HBINK bp, BINKFRAMEBUFFERS * set, U32 minimum_alignment );
+  #define BinkAllocateFrameBuffers indirectBinkAllocateFrameBuffers
+  RADDEFFUNC void * indirectBinkUtilMalloc(U64 bytes);
+  #define BinkUtilMalloc indirectBinkUtilMalloc
+  RADDEFFUNC void indirectBinkUtilFree(void * ptr);
+  #define BinkUtilFree indirectBinkUtilFree
+#endif
+
+static BINKTEXTURES * Create_textures( BINKSHADERS * pshaders, HBINK bink, void * user_ptr );
+static void Free_shaders( BINKSHADERS * pshaders );
+static void Free_textures( BINKTEXTURES * ptextures );
+static void Start_texture_update( BINKTEXTURES * ptextures );
+static void Finish_texture_update( BINKTEXTURES * ptextures );
+static void Draw_textures( BINKTEXTURES * ptextures, BINKSHADERS * pshaders, void * graphics_context );
+static void Set_draw_position( BINKTEXTURES * ptextures, F32 x0, F32 y0, F32 x1, F32 y1 );
+static void Set_draw_corners( BINKTEXTURES * ptextures, F32 Ax, F32 Ay, F32 Bx, F32 By, F32 Cx, F32 Cy );
+static void Set_source_rect( BINKTEXTURES * ptextures, F32 u0, F32 v0, F32 u1, F32 v1 );
+static void Set_alpha_settings( BINKTEXTURES * ptextures, F32 alpha_value, S32 draw_type );
+static void Set_hdr_settings( BINKTEXTURES * ptextures, S32 tonemap, F32 exposure, S32 out_nits  );
+
+//-----------------------------------------------------------------------------
+RADDEFFUNC BINKSHADERS * Create_Bink_shaders(void * dummy_device)
+{
+	BINKSHADERSRHI * pshaders;
+
+	pshaders = (BINKSHADERSRHI*)FMemory::Malloc(sizeof(*pshaders));
+	if (pshaders == 0)
+	{
+		return 0;
+	}
+
+	FMemory::Memset(pshaders, 0, sizeof(*pshaders));
+
+	pshaders->pub.Create_textures = Create_textures;
+	pshaders->pub.Free_shaders = Free_shaders;
+
+	return &pshaders->pub;
+}
+
+//-----------------------------------------------------------------------------
+static void Free_shaders( BINKSHADERS * pshaders )
+{
+	FMemory::Free(pshaders);
+}
+
+//-----------------------------------------------------------------------------
+static BINKTEXTURES * Create_textures(BINKSHADERS * pshaders, HBINK bink, void * user_ptr)
+{
+	BINKTEXTURESRHI * textures;
+	BINKFRAMEBUFFERS * bb;
+
+	textures = (BINKTEXTURESRHI *)FMemory::Malloc(sizeof(*textures));
+	if (textures == 0)
+		return 0;
+
+	FMemory::Memset(textures, 0, sizeof(*textures));
+
+	textures->pub.user_ptr = user_ptr;
+	textures->shaders = (BINKSHADERSRHI*)pshaders;
+	textures->bink = bink;
+	textures->video_width = bink->Width;
+	textures->video_height = bink->Height;
+
+	bb = &textures->bink_buffers;
+
+	BinkGetFrameBuffersInfo(bink, bb);
+
+#if !PLATFORM_HAS_DIRECT_TEXTURE_MEMORY_ACCESS
+	// allocate the system memory buffers if not allocated
+	if (!BinkAllocateFrameBuffers(bink, bb, 0))
+	{
+		FMemory::Free(textures);
+		return 0;
+	}
+#endif
+
+	FRHICommandListImmediate& RHICmdList = GetImmediateCommandList_ForRenderCommand();
+	EPixelFormat format = PF_R8;
+#if PLATFORM_HAS_DIRECT_TEXTURE_MEMORY_ACCESS 
+	ETextureCreateFlags TexCreateFlags = TexCreate_Dynamic | TexCreate_CPUReadback | TexCreate_CPUWritable | TexCreate_OfflineProcessed | TexCreate_NoTiling;
+#else
+	ETextureCreateFlags TexCreateFlags = TexCreate_Dynamic | TexCreate_NoTiling;
+#endif
+	FRHIResourceCreateInfo CreateInfo(TEXT("BINK"));
+
+	for (int i = 0; i < bb->TotalFrames; ++i)
+	{
+		BINKFRAMEPLANESET *bp_src = &bb->Frames[i];
+
+		if (bp_src->YPlane.Allocate)
+		{
+			textures->Ytexture[i] = RHICreateTexture2D(bb->YABufferWidth, bb->YABufferHeight, format, 1, 1, TexCreateFlags, CreateInfo);
+#if PLATFORM_HAS_DIRECT_TEXTURE_MEMORY_ACCESS 
+			bp_src->YPlane.Buffer = GDynamicRHI->LockTexture2D_RenderThread(RHICmdList, textures->Ytexture[i], 0, PLATFORM_DIRECT_TEXTURE_MEMORY_ACCESS_LOCK_MODE, bp_src->YPlane.BufferPitch, false, false);
+			bp_src->YPlane.BufferPitch = bp_src->YPlane.BufferPitch ? bp_src->YPlane.BufferPitch : ((bb->YABufferWidth+255)&-256);
+#endif
+		}
+
+		if (bp_src->cRPlane.Allocate)
+		{
+			textures->cRtexture[i] = RHICreateTexture2D(bb->cRcBBufferWidth, bb->cRcBBufferHeight, format, 1, 1, TexCreateFlags, CreateInfo);
+#if PLATFORM_HAS_DIRECT_TEXTURE_MEMORY_ACCESS 
+			bp_src->cRPlane.Buffer = GDynamicRHI->LockTexture2D_RenderThread(RHICmdList, textures->cRtexture[i], 0, PLATFORM_DIRECT_TEXTURE_MEMORY_ACCESS_LOCK_MODE, bp_src->cRPlane.BufferPitch, false, false);
+			bp_src->cRPlane.BufferPitch = bp_src->cRPlane.BufferPitch ? bp_src->cRPlane.BufferPitch : ((bb->cRcBBufferWidth+255)&-256);
+#endif
+		}
+
+		if (bp_src->cBPlane.Allocate)
+		{
+			textures->cBtexture[i] = RHICreateTexture2D(bb->cRcBBufferWidth, bb->cRcBBufferHeight, format, 1, 1, TexCreateFlags, CreateInfo);
+#if PLATFORM_HAS_DIRECT_TEXTURE_MEMORY_ACCESS 
+			bp_src->cBPlane.Buffer = GDynamicRHI->LockTexture2D_RenderThread(RHICmdList, textures->cBtexture[i], 0, PLATFORM_DIRECT_TEXTURE_MEMORY_ACCESS_LOCK_MODE, bp_src->cBPlane.BufferPitch, false, false);
+			bp_src->cBPlane.BufferPitch = bp_src->cBPlane.BufferPitch ? bp_src->cBPlane.BufferPitch : ((bb->cRcBBufferWidth+255)&-256);
+#endif
+		}
+
+		if (bp_src->APlane.Allocate)
+		{
+			textures->Atexture[i] = RHICreateTexture2D(bb->YABufferWidth, bb->YABufferHeight, format, 1, 1, TexCreateFlags, CreateInfo);
+#if PLATFORM_HAS_DIRECT_TEXTURE_MEMORY_ACCESS 
+			bp_src->APlane.Buffer = GDynamicRHI->LockTexture2D_RenderThread(RHICmdList, textures->Atexture[i], 0, PLATFORM_DIRECT_TEXTURE_MEMORY_ACCESS_LOCK_MODE, bp_src->APlane.BufferPitch, false, false);
+			bp_src->APlane.BufferPitch = bp_src->APlane.BufferPitch ? bp_src->APlane.BufferPitch : ((bb->YABufferWidth+255)&-256);
+#endif
+		}
+
+		if (bp_src->HPlane.Allocate)
+		{
+			textures->Htexture[i] = RHICreateTexture2D(bb->YABufferWidth, bb->YABufferHeight, format, 1, 1, TexCreateFlags, CreateInfo);
+#if PLATFORM_HAS_DIRECT_TEXTURE_MEMORY_ACCESS 
+			bp_src->HPlane.Buffer = GDynamicRHI->LockTexture2D_RenderThread(RHICmdList, textures->Htexture[i], 0, PLATFORM_DIRECT_TEXTURE_MEMORY_ACCESS_LOCK_MODE, bp_src->HPlane.BufferPitch, false, false);
+			bp_src->HPlane.BufferPitch = bp_src->HPlane.BufferPitch ? bp_src->HPlane.BufferPitch : ((bb->YABufferWidth+255)&-256);
+#endif
+		}
+	}
+
+	// Register our locked texture pointers with Bink
+	BinkRegisterFrameBuffers(bink, bb);
+
+	Set_draw_corners(&textures->pub, 0, 0, 1, 0, 0, 1);
+	Set_source_rect(&textures->pub, 0, 0, 1, 1);
+	Set_alpha_settings(&textures->pub, 1, 0);
+	Set_hdr_settings(&textures->pub, 0, 1.0f, 80);
+
+	textures->pub.Free_textures = Free_textures;
+	textures->pub.Start_texture_update = Start_texture_update;
+	textures->pub.Finish_texture_update = Finish_texture_update;
+	textures->pub.Draw_textures = Draw_textures;
+	textures->pub.Set_draw_position = Set_draw_position;
+	textures->pub.Set_draw_corners = Set_draw_corners;
+	textures->pub.Set_source_rect = Set_source_rect;
+	textures->pub.Set_alpha_settings = Set_alpha_settings;
+	textures->pub.Set_hdr_settings = Set_hdr_settings;
+
+	return &textures->pub;
+}
+
+//-----------------------------------------------------------------------------
+
+static void Free_textures(BINKTEXTURES* ptextures)
+{
+	FRHICommandListImmediate& RHICmdList = GetImmediateCommandList_ForRenderCommand();
+	BINKTEXTURESRHI* textures = (BINKTEXTURESRHI*)ptextures;
+	BINKFRAMEBUFFERS *bb = &textures->bink_buffers;
+
+	/*
+	RHICmdList.ImmediateFlush(EImmediateFlushType::FlushRHIThreadFlushResources);
+	RHIFlushResources();
+	RHICmdList.SubmitCommandsHint();
+	FPlatformMisc::MemoryBarrier();
+	*/
+
+	for (int i = 0; i < bb->TotalFrames; ++i)
+	{
+		BINKFRAMEPLANESET *bp_src = &bb->Frames[i];
+#if PLATFORM_HAS_DIRECT_TEXTURE_MEMORY_ACCESS 
+		if (bp_src->YPlane.Buffer)  GDynamicRHI->UnlockTexture2D_RenderThread(RHICmdList, textures->Ytexture[i], 0, false);
+		if (bp_src->cRPlane.Buffer) GDynamicRHI->UnlockTexture2D_RenderThread(RHICmdList, textures->cRtexture[i], 0, false);
+		if (bp_src->cBPlane.Buffer) GDynamicRHI->UnlockTexture2D_RenderThread(RHICmdList, textures->cBtexture[i], 0, false);
+		if (bp_src->APlane.Buffer)  GDynamicRHI->UnlockTexture2D_RenderThread(RHICmdList, textures->Atexture[i], 0, false);
+		if (bp_src->HPlane.Buffer)  GDynamicRHI->UnlockTexture2D_RenderThread(RHICmdList, textures->Htexture[i], 0, false);
+#endif
+		if (textures->Ytexture[i].IsValid())  textures->Ytexture[i].SafeRelease();
+		if (textures->cRtexture[i].IsValid()) textures->cRtexture[i].SafeRelease();
+		if (textures->cBtexture[i].IsValid()) textures->cBtexture[i].SafeRelease();
+		if (textures->Atexture[i].IsValid())  textures->Atexture[i].SafeRelease();
+		if (textures->Htexture[i].IsValid())  textures->Htexture[i].SafeRelease();
+	}
+
+	FMemory::Free(textures);
+}
+
+static void Start_texture_update( BINKTEXTURES * ptextures )
+{
+	//BINKTEXTURESRHI *textures = (BINKTEXTURESRHI*)ptextures;
+}
+
+static void Finish_texture_update( BINKTEXTURES * ptextures )
+{
+	FRHICommandListImmediate& RHICmdList = GetImmediateCommandList_ForRenderCommand();
+	BINKTEXTURESRHI * textures = (BINKTEXTURESRHI*)ptextures;
+	BINKFRAMEBUFFERS *bb = &textures->bink_buffers;
+	int frame_num = bb->FrameNum;
+	BINKFRAMEPLANESET *bp_src = &bb->Frames[frame_num];
+
+#if !PLATFORM_HAS_DIRECT_TEXTURE_MEMORY_ACCESS
+	textures->dirty = 1;
+#endif
+
+	textures->pub.Ytexture = (void*)(UINTa)&textures->Ytexture[frame_num];
+	textures->pub.cRtexture = (void*)(UINTa)&textures->cRtexture[frame_num];
+	textures->pub.cBtexture = (void*)(UINTa)&textures->cBtexture[frame_num];
+	textures->pub.Atexture = (void*)(UINTa)&textures->Atexture[frame_num];
+	textures->pub.Htexture = (void*)(UINTa)&textures->Htexture[frame_num];
+}
+
+//-----------------------------------------------------------------------------
+
+static void update_plane_texture_rect(FRHICommandListImmediate& RHI, FRHITexture2D* RHITexture, BINKPLANE const* plane, unsigned w, unsigned h)
+{
+	uint32 Stride = 0;
+	unsigned char* TextureMemory = (unsigned char*)GDynamicRHI->LockTexture2D_RenderThread(RHI, RHITexture, 0, RLM_WriteOnly, Stride, false);
+
+	if (TextureMemory)
+	{
+		check(Stride >= w);
+		if (Stride == w)
+		{
+			FMemory::Memcpy(TextureMemory, plane->Buffer, w * h);
+		}
+		else
+		{
+			for (unsigned i = 0; i < h; ++i)
+			{
+				FMemory::Memcpy(TextureMemory + i * Stride, (char*)plane->Buffer + i * w, w);
+			}
+		}
+		GDynamicRHI->UnlockTexture2D_RenderThread(RHI, RHITexture, 0, false);
+	}
+}
+
+static void Draw_textures(BINKTEXTURES* ptextures, BINKSHADERS* pshaders, void* graphics_context)
+{
+	BINKTEXTURESRHI* textures = (BINKTEXTURESRHI*)ptextures;
+	BINKSHADERSRHI* bshaders = (BINKSHADERSRHI*)pshaders;
+	BINKFRAMEBUFFERS *bb = &textures->bink_buffers;
+	int frame_num = bb->FrameNum;
+	BINKFRAMEPLANESET *bp_src = &bb->Frames[frame_num];
+
+	// nowhere to render to?
+	if (BinkRHIRenderTarget == 0)
+		return;
+
+	if (bshaders == 0)
+		bshaders = textures->shaders;
+
+	int hasAPlane = bp_src->APlane.Allocate;
+	int hasHPlane = bp_src->HPlane.Allocate;
+
+	FMemMark Mark(FMemStack::Get());
+
+	FRHICommandListImmediate& RHICmdList = GetImmediateCommandList_ForRenderCommand();
+	FRDGBuilder BinkGraphBuilder(RHICmdList);
+
+	FBinkParameters *consts = BinkGraphBuilder.AllocParameters<FBinkParameters>();
+
+	FRDGTexture *BinkRDGTexture = BinkRegisterExternalTexture(BinkGraphBuilder, BinkRHIRenderTarget, TEXT("Bink_RT"));
+	consts->RenderTargets[0] = FRenderTargetBinding(BinkRDGTexture, BinkRenderTargetLoadAction, 0);
+	consts->tex0 = BinkRegisterExternalTexture(BinkGraphBuilder, textures->Ytexture[frame_num], TEXT("BinkY"));
+	consts->tex1 = BinkRegisterExternalTexture(BinkGraphBuilder, textures->cRtexture[frame_num], TEXT("BinkCr"));
+	consts->tex2 = BinkRegisterExternalTexture(BinkGraphBuilder, textures->cBtexture[frame_num], TEXT("BinkCb"));
+	consts->tex3 = BinkRegisterExternalTexture(BinkGraphBuilder, hasAPlane ? textures->Atexture[frame_num] : textures->Ytexture[frame_num], TEXT("BinkA"));
+	consts->tex4 = BinkRegisterExternalTexture(BinkGraphBuilder, hasHPlane ? textures->Htexture[frame_num] : textures->Ytexture[frame_num], TEXT("BinkH"));
+	consts->samp0 = TStaticSamplerState<SF_Bilinear>::GetRHI();
+	consts->samp1 = TStaticSamplerState<SF_Bilinear>::GetRHI();
+	consts->samp2 = TStaticSamplerState<SF_Bilinear>::GetRHI();
+	consts->samp3 = TStaticSamplerState<SF_Bilinear>::GetRHI();
+	consts->samp4 = TStaticSamplerState<SF_Bilinear>::GetRHI();
+
+	// Pixel shader alpha constants
+	if (textures->draw_type == 3)
+	{
+		unsigned BinkColorA = 0xffffffff;
+		unsigned BinkColorB = 0xffc0c0c0;
+		consts->consta.X = ((F32)(S32)((BinkColorA >> 0) & 255)) * (1.0f / 255.0f);
+		consts->consta.Y = ((F32)(S32)((BinkColorA >> 8) & 255)) * (1.0f / 255.0f);
+		consts->consta.Z = ((F32)(S32)((BinkColorA >> 16) & 255)) * (1.0f / 255.0f);
+		consts->consta.W = ((F32)(S32)((BinkColorA >> 24) & 255)) * (1.0f / 255.0f);
+		consts->crc.X = ((F32)(S32)((BinkColorB >> 0) & 255)) * (1.0f / 255.0f);
+		consts->crc.Y = ((F32)(S32)((BinkColorB >> 8) & 255)) * (1.0f / 255.0f);
+		consts->crc.Z = ((F32)(S32)((BinkColorB >> 16) & 255)) * (1.0f / 255.0f);
+		consts->crc.W = ((F32)(S32)((BinkColorB >> 24) & 255)) * (1.0f / 255.0f);
+		consts->cbc.X = ((F32)textures->bink->Width) / 8.0f;
+		consts->cbc.Y = ((F32)textures->bink->Height) / 8.0f;
+		consts->cbc.Z = 0.f;
+		consts->cbc.W = 0.f;
+	}
+	else
+	{
+		// Pixel shader alpha constants
+		consts->consta.Z = consts->consta.Y = consts->consta.X = textures->draw_type == 1 ? textures->alpha : 1.0f;
+		consts->consta.W = textures->alpha;
+
+		if (hasHPlane)
+		{
+			// HDR stuff
+			consts->crc.X = textures->bink->ColorSpace[0];
+			consts->crc.Y = textures->exposure;
+			consts->crc.Z = textures->out_luma;
+			consts->crc.W = 0;
+			consts->cbc.X = textures->bink->ColorSpace[1];
+			consts->cbc.Y = textures->bink->ColorSpace[2];
+			consts->cbc.Z = textures->bink->ColorSpace[3];
+			consts->cbc.W = textures->bink->ColorSpace[4];
+		}
+		else
+		{
+			// set the constants for the type of ycrcb we have
+			consts->crc.X = textures->bink->ColorSpace[0];
+			consts->crc.Y = textures->bink->ColorSpace[1];
+			consts->crc.Z = textures->bink->ColorSpace[2];
+			consts->crc.W = textures->bink->ColorSpace[3];
+			consts->cbc.X = textures->bink->ColorSpace[4];
+			consts->cbc.Y = textures->bink->ColorSpace[5];
+			consts->cbc.Z = textures->bink->ColorSpace[6];
+			consts->cbc.W = textures->bink->ColorSpace[7];
+			consts->adj.X = textures->bink->ColorSpace[8];
+			consts->adj.Y = textures->bink->ColorSpace[9];
+			consts->adj.Z = textures->bink->ColorSpace[10];
+			consts->adj.W = textures->bink->ColorSpace[11];
+			consts->yscale.X = textures->bink->ColorSpace[12];
+			consts->yscale.Y = textures->bink->ColorSpace[13];
+			consts->yscale.Z = textures->bink->ColorSpace[14];
+			consts->yscale.W = textures->bink->ColorSpace[15];
+		}
+	}
+
+	// Vertex shader constants
+	consts->xy_xform0.X = (textures->Bx - textures->Ax) * 2.0f;
+	consts->xy_xform0.Y = (textures->Cx - textures->Ax) * 2.0f;
+	consts->xy_xform0.Z = (textures->By - textures->Ay) * -2.0f;
+	consts->xy_xform0.W = (textures->Cy - textures->Ay) * -2.0f; // view space has +y = up, our coords have +y = down
+	consts->xy_xform1.X = textures->Ax * 2.0f - 1.0f;
+	consts->xy_xform1.Y = 1.0f - textures->Ay * 2.0f;
+	consts->xy_xform1.Z = 0.0f;
+	consts->xy_xform1.W = 0.0f;
+
+	// UV matrix
+	{
+		F32 luma_u_scale = (F32)textures->video_width / (F32)bb->YABufferWidth;
+		F32 luma_v_scale = (F32)textures->video_height / (F32)bb->YABufferHeight;
+		F32 chroma_u_scale = (F32)(textures->video_width / 2) / (F32)bb->cRcBBufferWidth;
+		F32 chroma_v_scale = (F32)(textures->video_height / 2) / (F32)bb->cRcBBufferHeight;
+
+		// Set up matrix columns for UV transform (currently just scale+translate, could add rotation)
+		// X column
+		consts->uv_xform0.X = (textures->u1 - textures->u0) * luma_u_scale;
+		consts->uv_xform0.Y = 0.0f;
+		consts->uv_xform0.Z = (textures->u1 - textures->u0) * chroma_u_scale;
+		consts->uv_xform0.W = 0.0f;
+
+		// Y column
+		consts->uv_xform1.X = 0.0f;
+		consts->uv_xform1.Y = (textures->v1 - textures->v0) * luma_v_scale;
+		consts->uv_xform1.Z = 0.0f;
+		consts->uv_xform1.W = (textures->v1 - textures->v0) * chroma_v_scale;
+
+		// W column (translation)
+		consts->uv_xform2.X = textures->u0 * luma_u_scale;
+		consts->uv_xform2.Y = textures->v0 * luma_v_scale;
+		consts->uv_xform2.Z = textures->u0 * chroma_u_scale;
+		consts->uv_xform2.W = textures->v0 * chroma_v_scale;
+	}
+
+	FBinkDrawVS::FParameters* vert_params = BinkGraphBuilder.AllocParameters<FBinkDrawVS::FParameters>();
+	vert_params->BinkParameters = *consts;
+
+	FBinkDrawICtCpPS::FParameters* ictcp_params = BinkGraphBuilder.AllocParameters<FBinkDrawICtCpPS::FParameters>();
+	ictcp_params->BinkParameters = *consts;
+
+	FBinkDrawYCbCrPS::FParameters* ycbcr_params = BinkGraphBuilder.AllocParameters<FBinkDrawYCbCrPS::FParameters>();
+	ycbcr_params->BinkParameters = *consts;
+
+	auto* ShaderMap = GetGlobalShaderMap(GMaxRHIFeatureLevel);
+	TShaderMapRef<FBinkDrawVS> DrawVS(ShaderMap);
+
+	FBinkDrawICtCpPS::FPermutationDomain ictcp_pv;
+	ictcp_pv.Set<FBinkDrawICtCpPS::FALPHA>(!!hasAPlane);
+	ictcp_pv.Set<FBinkDrawICtCpPS::FTONEMAP>(textures->tonemap == 1);
+	ictcp_pv.Set<FBinkDrawICtCpPS::FST2084>(textures->tonemap == 2);
+	TShaderMapRef<FBinkDrawICtCpPS> DrawICtCpPS(ShaderMap, ictcp_pv);
+
+	FBinkDrawYCbCrPS::FPermutationDomain ycbcr_pv;
+	ycbcr_pv.Set<FBinkDrawYCbCrPS::FALPHA>(!!hasAPlane);
+	ycbcr_pv.Set<FBinkDrawYCbCrPS::FSRGB>(!!(textures->draw_flags & 0x80000000));
+	TShaderMapRef<FBinkDrawYCbCrPS> DrawYCbCrPS(ShaderMap, ycbcr_pv);
+
+	FRDGEventName EventName(TEXT("Bink"));
+	BinkGraphBuilder.AddPass(
+		MoveTemp(EventName),
+		consts,
+		ERDGPassFlags::Raster,
+		[&](FRHICommandListImmediate& RHICmdList)
+		{
+			if (textures->dirty)
+			{
+				textures->dirty = 0;
+
+				FUpdateTextureRegion2D region_YAH(0, 0, 0, 0, bb->YABufferWidth, bb->YABufferHeight);
+				FUpdateTextureRegion2D region_cRcB(0, 0, 0, 0, bb->cRcBBufferWidth, bb->cRcBBufferHeight);
+
+				GDynamicRHI->UpdateTexture2D_RenderThread(RHICmdList, textures->Ytexture[frame_num], 0, region_YAH, bp_src->YPlane.BufferPitch, (uint8*)bp_src->YPlane.Buffer);
+				GDynamicRHI->UpdateTexture2D_RenderThread(RHICmdList, textures->cRtexture[frame_num], 0, region_cRcB, bp_src->cRPlane.BufferPitch, (uint8*)bp_src->cRPlane.Buffer);
+				GDynamicRHI->UpdateTexture2D_RenderThread(RHICmdList, textures->cBtexture[frame_num], 0, region_cRcB, bp_src->cBPlane.BufferPitch, (uint8*)bp_src->cBPlane.Buffer);
+				if (hasAPlane) GDynamicRHI->UpdateTexture2D_RenderThread(RHICmdList, textures->Atexture[frame_num], 0, region_YAH, bp_src->APlane.BufferPitch, (uint8*)bp_src->APlane.Buffer);
+				if (hasHPlane) GDynamicRHI->UpdateTexture2D_RenderThread(RHICmdList, textures->Htexture[frame_num], 0, region_YAH, bp_src->HPlane.BufferPitch, (uint8*)bp_src->HPlane.Buffer);
+			}
+
+			FGraphicsPipelineStateInitializer GraphicsPSOInit;
+			RHICmdList.ApplyCachedRenderTargets(GraphicsPSOInit);
+
+			// Set up blending 
+			if ((!hasAPlane && textures->alpha >= 0.999f) || textures->draw_type == 2)
+			{
+				// opaque
+				GraphicsPSOInit.BlendState = TStaticBlendState<>::GetRHI();
+			}
+			else if (textures->draw_type == 1)
+			{
+				// alpha pre-multiplied
+				GraphicsPSOInit.BlendState = TStaticBlendState<CW_RGBA, BO_Add, BF_One, BF_InverseSourceAlpha, BO_Add, BF_One, BF_Zero>::GetRHI();
+			}
+			else
+			{
+				// normal alpha
+				GraphicsPSOInit.BlendState = TStaticBlendState<CW_RGBA, BO_Add, BF_SourceAlpha, BF_InverseSourceAlpha, BO_Add, BF_One, BF_Zero>::GetRHI();
+			}
+
+			// Disable backface culling
+			GraphicsPSOInit.RasterizerState = TStaticRasterizerState<FM_Solid, CM_None, false, false>::GetRHI();
+
+			// Disable Depth/Stencil test
+			GraphicsPSOInit.DepthStencilState = TStaticDepthStencilState<false, CF_Always>::GetRHI();
+
+			FVertexDeclarationElementList Elements;
+			GraphicsPSOInit.BoundShaderState.VertexDeclarationRHI = PipelineStateCache::GetOrCreateVertexDeclaration(Elements);
+			GraphicsPSOInit.BoundShaderState.VertexShaderRHI = DrawVS.GetVertexShader();
+			GraphicsPSOInit.PrimitiveType = PT_TriangleStrip;
+
+			// Set VS/PS shaders
+			if (hasHPlane)
+			{
+				GraphicsPSOInit.BoundShaderState.PixelShaderRHI = DrawICtCpPS.GetPixelShader();
+				SetGraphicsPipelineState(RHICmdList, GraphicsPSOInit, 0);
+				SetShaderParameters(RHICmdList, DrawICtCpPS, DrawICtCpPS.GetPixelShader(), *ictcp_params);
+			}
+			else
+			{
+				GraphicsPSOInit.BoundShaderState.PixelShaderRHI = DrawYCbCrPS.GetPixelShader();
+				SetGraphicsPipelineState(RHICmdList, GraphicsPSOInit, 0);
+				SetShaderParameters(RHICmdList, DrawYCbCrPS, DrawYCbCrPS.GetPixelShader(), *ycbcr_params);
+			}
+
+			SetShaderParameters(RHICmdList, DrawVS, DrawVS.GetVertexShader(), *vert_params);
+
+			RHICmdList.DrawPrimitive(0, 2, 1);
+		});
+
+	BinkGraphBuilder.Execute();
+}
+
+static void Set_draw_position(BINKTEXTURES * ptextures, F32 x0, F32 y0, F32 x1, F32 y1)
+{
+	BINKTEXTURESRHI * textures = (BINKTEXTURESRHI*)ptextures;
+	textures->Ax = x0;
+	textures->Ay = y0;
+	textures->Bx = x1;
+	textures->By = y0;
+	textures->Cx = x0;
+	textures->Cy = y1;
+}
+
+static void Set_draw_corners(BINKTEXTURES * ptextures, F32 Ax, F32 Ay, F32 Bx, F32 By, F32 Cx, F32 Cy)
+{
+	BINKTEXTURESRHI * textures = (BINKTEXTURESRHI*)ptextures;
+	textures->Ax = Ax;
+	textures->Ay = Ay;
+	textures->Bx = Bx;
+	textures->By = By;
+	textures->Cx = Cx;
+	textures->Cy = Cy;
+}
+
+static void Set_source_rect(BINKTEXTURES * ptextures, F32 u0, F32 v0, F32 u1, F32 v1)
+{
+	BINKTEXTURESRHI * textures = (BINKTEXTURESRHI*)ptextures;
+	textures->u0 = u0;
+	textures->v0 = v0;
+	textures->u1 = u1;
+	textures->v1 = v1;
+}
+
+static void Set_alpha_settings(BINKTEXTURES * ptextures, F32 alpha_value, S32 draw_type)
+{
+	BINKTEXTURESRHI * textures = (BINKTEXTURESRHI*)ptextures;
+	textures->alpha = alpha_value;
+	textures->draw_type = draw_type & 0x0FFFFFFF;
+	textures->draw_flags = draw_type & 0xF0000000;
+}
+
+static void Set_hdr_settings(BINKTEXTURES * ptextures, S32 tonemap, F32 exposure, S32 out_nits)
+{
+	BINKTEXTURESRHI * textures = (BINKTEXTURESRHI*)ptextures;
+	textures->tonemap = tonemap;
+	textures->exposure = exposure;
+	textures->out_luma = ((F32)out_nits) / 80.f;
+}
+
diff --git a/Engine/Plugins/Media/BinkMedia/Source/BinkMediaPlayer/Public/BinkMediaPlayer.h b/Engine/Plugins/Media/BinkMedia/Source/BinkMediaPlayer/Public/BinkMediaPlayer.h
--- a/Engine/Plugins/Media/BinkMedia/Source/BinkMediaPlayer/Public/BinkMediaPlayer.h
+++ b/Engine/Plugins/Media/BinkMedia/Source/BinkMediaPlayer/Public/BinkMediaPlayer.h
@@ -461,7 +461,7 @@ public:
 	bool Open(const FString& Url);
 	void Close();
 	bool IsReady() const { return bnk != NULL; }
-	void UpdateTexture(FRHICommandListImmediate &RHICmdList, FTextureRHIRef ref, void *nativePtr, int width, int height, bool isEditor, bool tonemap, int output_nits, float alpha, bool srgb_decode, bool is_hdr);
+	void UpdateTexture(FRHICommandListImmediate &RHICmdList, FTexture2DRHIRef ref, void *nativePtr, int width, int height, bool isEditor, bool tonemap, int output_nits, float alpha, bool srgb_decode, bool is_hdr);
 
 	FIntPoint GetDimensions() const 
 	{
diff --git a/Engine/Plugins/Media/BinkMedia/Source/BinkMediaPlayerEditor/Private/Models/BinkMediaPlayerEditorTexture.h b/Engine/Plugins/Media/BinkMedia/Source/BinkMediaPlayerEditor/Private/Models/BinkMediaPlayerEditorTexture.h
--- a/Engine/Plugins/Media/BinkMedia/Source/BinkMediaPlayerEditor/Private/Models/BinkMediaPlayerEditorTexture.h
+++ b/Engine/Plugins/Media/BinkMedia/Source/BinkMediaPlayerEditor/Private/Models/BinkMediaPlayerEditorTexture.h
@@ -32,18 +32,7 @@ struct FBinkMediaPlayerEditorTexture : public FTickableGameObject
 				return;
 			}
 
-			uint32 width = tex->GetSizeX();
-			uint32 height = tex->GetSizeY();
-			void* nativePtr = tex->GetNativeResource();
-			FRHIRenderPassInfo RPInfo(SlateTexture->GetRHIRef(), ERenderTargetActions::Load_Store);
-			TransitionRenderPassTargets(RHICmdList, RPInfo);
-			RHICmdList.BeginRenderPass(RPInfo, TEXT("RenderBink"));
-			RHICmdList.SetViewport(0, 0, 0.0f, width, height, 1.0f);
-			RHICmdList.SubmitCommandsHint();
-			MediaPlayer->UpdateTexture(RHICmdList, static_cast<FTextureRHIRef>(tex), nativePtr, width, height, true, true, 80, 1, false, false);
-			RHICmdList.SubmitCommandsHint();
-			RHICmdList.Transition(FRHITransitionInfo(SlateTexture->GetRHIRef(), ERHIAccess::Unknown, ERHIAccess::UAVGraphics));
-			RHICmdList.EndRenderPass();
+			MediaPlayer->UpdateTexture(RHICmdList, tex, tex->GetNativeResource(), tex->GetSizeX(), tex->GetSizeY(), true, true, 80, 1, false, false);
 		});
 	}
 	
diff --git a/Engine/Plugins/Media/BinkMedia/Source/SDK/include/binkplugin.h b/Engine/Plugins/Media/BinkMedia/Source/SDK/include/binkplugin.h
--- a/Engine/Plugins/Media/BinkMedia/Source/SDK/include/binkplugin.h
+++ b/Engine/Plugins/Media/BinkMedia/Source/SDK/include/binkplugin.h
@@ -35,6 +35,7 @@
 
 enum BINKPLUGINAPI
 {
+#ifndef BUILDING_FOR_UNREAL_ONLY
   BinkGL = 0,
   BinkD3D9 = 1,
   BinkD3D11 = 2,
@@ -42,6 +43,9 @@ enum BINKPLUGINAPI
   BinkMetal = 4,
   BinkVulkan = 5,
   BinkNDA = 6,     // any nda gpu api
+#else
+  BinkRHI = 0,     // Unreal RHI api
+#endif
 };
 
 // some platforms need allocators for cpu/gpu memory
diff --git a/Engine/Plugins/Media/BinkMedia/Source/SDK/include/binktextures.h b/Engine/Plugins/Media/BinkMedia/Source/SDK/include/binktextures.h
new file mode 100644
--- /dev/null
+++ b/Engine/Plugins/Media/BinkMedia/Source/SDK/include/binktextures.h
@@ -0,0 +1,477 @@
+// Copyright Epic Games, Inc. All Rights Reserved.
+// Licensed under the terms of a valid Unreal Engine license agreement,
+//   and the separate 'Unreal Engine End User License Agreement for Publishing'.
+
+/*
+This is Epics's high level API for using 3D hardware to do Bink GPU-assisted
+decoding and color conversion. It is supported on every 3D platform that 
+Bink runs on.
+
+It's a simple API, so you should see your platform's example code
+to see all the nitty-gritty details.
+
+
+There two functions to start and stop Bink textures:
+
+  Create_Bink_shaders:  This function creates the pixel shaders we use to
+    do the color space conversion. On most platforms, you need to provide
+    a pointer to GPU device pointer. 
+    All of the other functions hang off the object that is returned here.
+
+  Free_Bink_shaders:  Frees the pixel shaders created in Create_Bink_shaders.
+
+There are four textures functions:
+
+  Create_Bink_textures:  This function takes BINK object, creates the texture
+    resources to render it quickly and binds them together for rendering.
+  
+  Free_Bink_textures:  Frees the resources allocated in Create_Bink_textures.
+
+  Draw_Bink_textures:  Renders the textures onto the screen.
+
+  Start_Bink_texture_update:  This function starts the texture updating process.
+    It should be called before you call BinkDoFrame (synchronous or async).
+    On some platforms, this locks/maps the texture for updating; on consoles, it
+    makes sure that the GPU is done reading from the texture buffers that we're
+    about to overwrite.
+
+  Finish_Bink_texture_update:  This function finishes the texture updating process.
+    It should be called after BinkDoFrame processing is done (in async mode: after
+    BinkDoFrameAsyncWait returns true). This will unlock/unmap textures (where it
+    applies) or issue a blit to transfer the new image data. On most consoles with
+    unified memory, we decode directly into GPU-readable memory, and this operation
+    is a no-op.
+
+On Windows D3D9, there are two other functions:
+
+  Before_Reset_Bink_textures: Call before you call the D3D Reset function.
+
+  After_Reset_Bink_textures: Call after you call the D3D Reset function.
+
+For Windows GL or D3D11, you simply free and re-create the textures completely
+  whenever the device is lost or reset.
+
+
+
+So, basically, playback works like this:
+
+  1) Call Create_Bink_shaders to create the pixel shaders.
+  
+  2) Open the Bink file with the BINKNOFRAMEBUFFERS flag (the Binktexture
+     implementations will allocate the memory to decompress into).
+
+  3) Call Create_Bink_textures to create the textures and bind the textures
+     into the Bink.
+
+  4) Call Start_Bink_texture_update before BinkDoFrame.
+
+  5) Call BinkDoFrame to decompress a video frame.
+
+  6) Call Finish_Bink_texture_update after BinkDoFrame (or after
+     BinkDoFrameAsyncWait returns true, if you're using async decoding).
+
+  7) Draw the frame using Draw_Bink_textures.
+
+
+And that's it! (Skipping over a few details - see the examples for all 
+the details...)
+
+Should drop in really quickly and it hides a ton of platform specific ugliness!
+*/
+
+#ifndef BINKTEXTURESH
+#define BINKTEXTURESH
+
+#ifndef __BINKH__
+#include "bink.h"
+#endif
+
+RADDEFSTART
+
+//=============================================================================
+
+// On some platforms, you might have one app that uses BinkTextures with multiple
+//   graphics APIs - these defines let you have multiple providers in one
+//   EXE by suffixing the names of D3D9 and D3D11 (GL gets the non-suffix).
+// This does mean that you must include D3D9 and/or D3D11 headers before
+//   including BinkTextures.h!
+
+#if defined( BINKTEXTURESGPUAPITYPE )
+
+  #if defined( __d3d12_x_h__ ) || defined( __d3d12_xs_h__ ) ||defined( __d3d12_h__ ) || defined( BINKD3D12FUNCTIONS ) || defined( BINKD3D12GPUFUNCTIONS )
+    #ifdef BINKD3D12GPUFUNCTIONS
+      #define BINKTEXTURESSUFFIX D3D12GPU
+    #else
+      #define BINKTEXTURESSUFFIX D3D12
+    #endif
+  #elif defined( __d3d11_x_h__ ) || defined( __d3d11_h__ ) || defined( BINKD3D11FUNCTIONS ) || defined( BINKD3D11GPUFUNCTIONS )
+    #ifdef BINKD3D11GPUFUNCTIONS
+      #define BINKTEXTURESSUFFIX D3D11GPU
+    #else
+      #define BINKTEXTURESSUFFIX D3D11
+    #endif
+  #elif defined( _D3D9_H_ ) || defined( BINKD3D9FUNCTIONS )
+    #define BINKTEXTURESSUFFIX D3D9
+  #elif defined( BINKMETALFUNCTIONS )
+    #define BINKTEXTURESSUFFIX Metal
+  #elif defined( BINKVULKANFUNCTIONS )
+    #define BINKTEXTURESSUFFIX Vulkan
+  #elif defined(__gl_es20_h_) || defined(__gl_h_) || defined(__gl2_h_) || defined(__gl3_h_) || defined( BINKGLFUNCTIONS ) || defined(BINKGLGPUFUNCTIONS)
+    #ifdef BINKGLGPUFUNCTIONS
+      #define BINKTEXTURESSUFFIX GLGPU
+    #else
+      #define BINKTEXTURESSUFFIX GL
+    #endif
+  #elif defined(BINKNDAFUNCTIONS)
+    #define BINKTEXTURESSUFFIX BINKNDAFUNCTIONS
+  #elif defined(BINKRHIFUNCTIONS)
+    #define BINKTEXTURESSUFFIX RHI
+  #else
+    #ifndef BINKTEXTURESSUFFIX
+      #define BINKTEXTURESSUFFIX Null
+    #endif
+  #endif
+
+  #define Create_Bink_shaders        RR_STRING_JOIN( Create_Bink_shaders,       BINKTEXTURESSUFFIX )
+
+#endif
+
+typedef struct BINKSHADERS BINKSHADERS;   // defined below
+typedef struct BINKTEXTURES BINKTEXTURES; // defined below
+
+//=============================================================================
+
+// For some APIs, you need to pass structs as argument to our functions.
+// This is where they get defined.
+
+#if ( defined( __d3d12_h__ ) || defined( __d3d12_x_h__ ) || defined( __d3d12_xs_h__ ) ) && !defined( BINKTEXTURESD3D12STRUCTSDEF )
+
+// BinkGPU can call you back at the beginning/end of its compute command lists in case you want to
+// insert timestamp queries (or other markup) into the command list.
+typedef void BINKGPUCOMMANDLISTCALLBACKD3D12( BINKTEXTURES * bink_tex, ID3D12GraphicsCommandList * cl, BOOL is_end, void * user_ptr );
+
+// This is what you pass as "device" to Create_shaders.
+typedef struct BINKCREATESHADERSD3D12 {
+  // The D3D12 device to use
+  ID3D12Device * device;
+  
+  // Command queue to submit work to in BinkGPU mode. Ignored in CPU decoding mode.
+  ID3D12CommandQueue * gpu_mode_command_queue;
+
+  // The fence to use for synchronization. Needs to be updated from app code!
+  ID3D12Fence * fence; 
+
+  // Prototype graphics pipeline state. Create Bink pipelines for the
+  // specified set of rasterizer state (multisampling or not),
+  // render target formats, etc. See binktexturesD3D12.cpp for
+  // details.
+  D3D12_GRAPHICS_PIPELINE_STATE_DESC prototype;
+
+  // Callback to run at beginning/end of BinkGPU command list. Ignored in CPU decoding mode.
+  BINKGPUCOMMANDLISTCALLBACKD3D12 * gpu_command_list_callback;
+  void * gpu_command_list_user_ptr;
+} BINKCREATESHADERSD3D12;
+
+// This is what you pass as context to Draw_textures.
+typedef struct BINKAPICONTEXTD3D12 {
+  // ---- Input parameters you provide
+
+  // Command list to render to
+  ID3D12GraphicsCommandList * command_list; // Command list to render to
+
+  // Fence value you're going to write via the command queue after executing
+  // the command_list. The sequence of fence values written needs to be
+  // monotonically increasing over time. This is what we use to synchronize
+  // buffer upload ands resource management!
+  UINT64 fence_value;
+
+  // ---- Output parameters we return
+
+  // In GPU decode mode, the draw command list may only run after decoding
+  // the frame has finished. Thus you need to call "Wait" on your graphics
+  // command queue to make sure our decode fence is signaled to the correct
+  // value before decoding can start.
+  //
+  // NOTE: "out_gpu_decode_fence" is not AddRef'd before being returned, so please
+  // don't call Release on it!
+
+  ID3D12Fence * out_gpu_decode_fence;   // In GPU mode, this is the fence you have to call wait on before submitting command_list.
+  UINT64 out_gpu_decode_fence_val;      // In GPU mode, this is the fence value you have to wait for before submitting command_list.
+} BINKAPICONTEXTD3D12;
+
+// We promise not to request more SRV descriptors than this per Bink, for any D3D12 BinkTextures implementation.
+// So if you want to use a simple free list allocator for Bink video descriptors, you can just use this as your
+// allocation granularity.
+#define BINKTEXTURESD3D12_MAX_SRVS_PER_BINK     128
+
+#define BINKTEXTURESD3D12STRUCTSDEF
+
+#endif
+
+#ifndef BINKCREATESHADERSPIPEANDQUEUEd
+#define BINKCREATESHADERSPIPEANDQUEUEd
+
+// Create_Bink_Shaders for some nda GPU implementations allows you to select which async compute pipe and queue ID to run
+// decoding jobs on. If not specified, we use our default. (At the time of writing, this is pipe 6, queue 1).
+// Fill out this struct and pass a pointer to it as "device" to Create_Bink_shaders.
+typedef struct BINKCREATESHADERSPIPEANDQUEUE 
+{
+  // Which async compute pipe to use for BinkGPU. Ignored for CPU decoding.
+  // Corresponds to the "pipeId" in ComputeQueue::initialize.
+  U32 pipe_id;
+
+  // Which queue to use for BinkGPU on the given compute pipe. Ignored for CPU decoding.
+  // Corresponds to the "queueId" in ComputeQueue::initialize.
+  U32 queue_id;
+} BINKCREATESHADERSPIPEANDQUEUE;
+
+#endif
+
+#ifndef BINKTEXTURES_METAL_DEF
+#define BINKTEXTURES_METAL_DEF
+
+// Create_Bink_Shaders for Metal needs a few more parameters - pass a pointer to this structure.
+
+typedef struct BINKCREATESHADERSMETAL 
+{
+  void * mtlDevice;
+  void * stateDesc[2];          // MTLRenderPipelineDescriptor *
+  void * mtlQueue;              // used on GPU-assist Bink Metal only
+  U32 render_target_formats[2]; // MTLPixelFormat(s) that we can render to
+} BINKCREATESHADERSMETAL;
+
+
+typedef struct BINKAPICONTEXTMETAL
+{
+  void * command_buffer;  // id <MTLCommandBuffer> 
+  void * command_encoder; // id <MTLRenderCommandEncoder>
+  void * frame_texture;   // View.currentDrawable.texture
+  U32 format_idx;         // Index into render_target_formats. See BINKCREATESHADERSMETAL
+} BINKAPICONTEXTMETAL;
+
+#endif
+
+#ifndef BINKTEXTURES_VULKAN_DEF
+#define BINKTEXTURES_VULKAN_DEF
+typedef struct BINKCREATESHADERSVULKAN 
+{
+  void *physical_device;        // VkPhysicalDevice
+  void *logical_device;         // VkDevice
+  void *queue;                  // VkQueue (needs graphics+compute capabilities where compute is only used for gpu decode)
+  U32 render_target_formats[2]; // VkFormat(s) that we can render to
+} BINKCREATESHADERSVULKAN;
+
+typedef struct BINKAPICONTEXTVULKAN
+{
+  void * cmdbuf;  // VkCommandBuffer
+  U32 format_idx; // Index into render_target_formats. See BINKCREATESHADERSVULKAN
+} BINKAPICONTEXTVULKAN;
+#endif
+
+//=============================================================================
+
+// Creates shaders that Bink uses.
+//   The value that you pass for device depends on the graphics API:
+//     D3D9 -    The d3d9 device (type LPDIRECT3DDEVICE9)
+//     D3D11 -   The d3d11 device (type D3D11Device *)
+//     D3D12 -   Pointer to a filled-out BINKCREATESHADERSD3D12 struct
+//     GL -      Parameter not used, just pass 0
+//     PS4 -     Pass 0 for default parameters, or a pointer to a filled-out BINKCREATESHADERSPS4 struct
+//     Metal   - Pointer to a filled out BINKCREATESHADERSMETAL struct
+BINKSHADERS * Create_Bink_shaders( void * device );
+
+// free our shaders
+typedef void Free_Bink_shaders_ft( BINKSHADERS * shaders );
+
+//=============================================================================
+
+// allocate the textures that we'll need - shaders defines the device we use to
+//   create them. The user_ptr gets stored in the BINKTEXTURES and also passed
+//   to any underlying "alloc" functions we call.
+typedef BINKTEXTURES * Create_Bink_textures_ft( BINKSHADERS * shaders, HBINK bink, void *user_ptr );
+
+//=============================================================================
+
+// frees the textures
+typedef void Free_Bink_textures_ft( BINKTEXTURES * textures );
+
+//=============================================================================
+// Bink texture updating
+
+// Start updating textures for this frame. Call before DoFrame.
+typedef void Start_Bink_texture_update_ft( BINKTEXTURES * textures );
+
+// Finish updating textures for this frame. Call after DoFrame is done.
+typedef void Finish_Bink_texture_update_ft( BINKTEXTURES * textures );
+
+
+//=============================================================================
+
+// Bink texture drawing functions
+
+// Draws the textures as a screen-aligned quad.
+
+//   If you pass a zero for "shaders", we'll use the shader used to create the
+//       textures (which is normally what you want).
+//   The value that you pass for "graphics_context" depends on the graphics API:
+//     D3D9 -    Parameter not used, just pass 0
+//     D3D11 -   The d3d11 device context (type ID3D11DeviceContext *)
+//     D3D12 -   Pointer to a BINKAPICONTEXTD3D12 struct with all fields filled out
+//     GL -      Parameter not used, just pass 0
+//     PS3 -     The Gcm context (usually, gCellGcmCurrentContext, type CellGcmContextData *)
+//     PS4 -     The Gfx context pointer to use (type sce::Gnmx::GfxContext *)
+//     Xbox360 - Parameter not used, just pass 0
+//     XboxOne - The Xbox device context (type ID3D11DeviceContextX *)
+//     NX -      The NX context pointer to use (type nn::gfx::CommandBuffer *)
+//     Metal -   Pointer to a BINKAPICONTEXTMETAL struct with all the fields filled out
+//     PS5 -     Pointer to the DrawCommandBuffer to use (no StateBuffer or Binder required)
+typedef void Draw_Bink_textures_ft( BINKTEXTURES * textures,
+                                    BINKSHADERS * shaders,
+                                    void * graphics_context );
+
+// Set the position to draw the video within the viewport.
+//   If you video doesn't move, you can call this once,
+//     otherwise call before each call to Draw_Bink_texture.
+//   By default, the video fills the video port.
+//
+//   0,0, 1,1 is the default which completely fills the viewport
+typedef void Set_Bink_draw_position_ft( BINKTEXTURES * textures,
+                                        F32 x0, F32 y0, F32 x1, F32 y1 );
+
+
+//   0,0, 1,0, 0,1 is the default which completely fills the viewport
+typedef void Set_Bink_draw_corners_ft( BINKTEXTURES * textures,
+                                        F32 Ax, F32 Ay, F32 Bx, F32 By, F32 Cx, F32 Cy );
+
+
+// Set draw source rect (from 0 to 1, in texturespace, so 0,0, 1,1 is whole texture)
+typedef void Set_Bink_source_rect_ft( BINKTEXTURES * textures,
+                                      F32 u0, F32 v0, F32 u1, F32 v1 );
+
+// Or these flags with draw_type_and_flags
+#define BINKDRAWFLAG_SRGB_DECODE (0x80000000)
+
+// Set the alpha settings for drawing.
+//   alpha_value is just a constant blend value for entire video frame
+//   draw_type is 0 = normal alpha, 1 = pre-multiplied alpha, 2 = opaque & copy alpha
+//   draw_type also contains flags in the high bits:
+//     BINKDRAWFLAG_SRGB_DECODE - decode sRGB to linear when rendering the video.
+typedef void Set_Bink_alpha_settings_ft( BINKTEXTURES * textures,
+                                         F32 alpha_value,
+                                         S32 draw_type_and_flags );
+
+// Set the hdr settings for drawing.
+//   tonemap specifies whether you want linear output (0), or filmic tonemapped output (1), or ST2084 HDR output (2).
+//   exposure is a scaling factor that happens before tonemapping (1.0=normal, <1.0 darken, >1.0 brighten)
+//   out_luma specifies the maximum luminance to output when tonemapping
+typedef void Set_Bink_hdr_settings_ft( BINKTEXTURES * texutres,
+                                       S32 tonemap,
+                                       F32 exposure, 
+                                       S32 output_nits );
+
+
+//=============================================================================
+
+// For D3D9 only, these functions are called before and after you reset the D3D9 device
+typedef void Before_Reset_Bink_textures_ft( BINKTEXTURES * textures );
+typedef S32  After_Reset_Bink_textures_ft( BINKTEXTURES * textures );
+
+//=============================================================================
+
+#ifndef BINKTEXTURESTRUCTS
+#define BINKTEXTURESTRUCTS
+
+#define Create_Bink_textures( shader, bink, user_ptr )   (shader)->Create_textures( shader, bink, user_ptr )
+#define Free_Bink_shaders( shader )                      (shader)->Free_shaders( shader )
+
+#define Start_Bink_texture_update( textures )            (textures)->Start_texture_update( textures )
+#define Finish_Bink_texture_update( textures )           (textures)->Finish_texture_update( textures )
+#define Draw_Bink_textures( textures, shaders, gr_ctx )  (textures)->Draw_textures( textures, shaders, gr_ctx )
+#define Set_Bink_draw_position( textures, x0,y0, x1,y1 ) (textures)->Set_draw_position( textures, x0,y0, x1,y1 )
+#define Set_Bink_draw_corners( textures, Ax,Ay, Bx,By, Cx,Cy ) (textures)->Set_draw_corners( textures, Ax,Ay, Bx,By, Cx,Cy )
+#define Set_Bink_source_rect( textures, u0,v0, u1,v1 )   (textures)->Set_source_rect( textures, u0,v0, u1,v1 )
+#define Set_Bink_alpha_settings( textures, alpha, draw_type_and_flags ) (textures)->Set_alpha_settings( textures, alpha, draw_type_and_flags )
+#define Set_Bink_hdr_settings( textures, tonemap, exposure, out_nits ) (textures)->Set_hdr_settings( textures, tonemap, exposure, out_nits )
+#define Before_Reset_Bink_textures( textures )           {if ((textures)->Before_Reset_textures) (textures)->Before_Reset_textures( textures );}
+#define After_Reset_Bink_textures( textures )            {if ((textures)->After_Reset_textures) (textures)->After_Reset_textures( textures );}
+#define Free_Bink_textures( textures )                   (textures)->Free_textures( textures )
+
+struct BINKSHADERS
+{
+  Create_Bink_textures_ft * Create_textures;
+  Free_Bink_shaders_ft * Free_shaders;
+
+  // you can use this for whatever you want
+  UINTa user_data[ 4 ];
+};
+
+struct BINKTEXTURES
+{
+  // wrap around doframe
+  Start_Bink_texture_update_ft * Start_texture_update;
+  Finish_Bink_texture_update_ft * Finish_texture_update;
+
+  // draw related functions
+  Draw_Bink_textures_ft * Draw_textures;
+  Set_Bink_draw_position_ft * Set_draw_position;
+  Set_Bink_draw_corners_ft * Set_draw_corners;
+  Set_Bink_source_rect_ft * Set_source_rect;
+  Set_Bink_alpha_settings_ft * Set_alpha_settings;
+  Set_Bink_hdr_settings_ft * Set_hdr_settings;
+
+  // D3D9 only (everywhere else is null)
+  Before_Reset_Bink_textures_ft * Before_Reset_textures;
+  After_Reset_Bink_textures_ft * After_Reset_textures;
+
+  Free_Bink_textures_ft * Free_textures;
+
+  // user pointer specified at Create_Bink_textures time
+  void * user_ptr;
+
+  // these are the platform specific texture pointers that you can draw manually with
+  void * Ytexture;
+  void * cRtexture;
+  void * cBtexture;
+  void * Atexture;
+  void * Htexture;
+
+  // you can use this for whatever you want
+  UINTa user_data[ 4 ];
+};
+
+#endif
+//=============================================================================
+
+
+// optional clean up, for when including the header multiple times for different APIs
+#if defined(BINKTEXTURESCLEANUP)
+  #undef Create_Bink_shaders 
+
+  #ifdef BINKD3D9FUNCTIONS
+    #undef BINKD3D9FUNCTIONS
+  #endif
+  #ifdef BINKD3D11FUNCTIONS
+    #undef BINKD3D11FUNCTIONS
+  #endif
+  #ifdef BINKD3D11GPUFUNCTIONS
+    #undef BINKD3D11GPUFUNCTIONS
+  #endif
+  #ifdef BINKD3D12FUNCTIONS
+    #undef BINKD3D12FUNCTIONS
+  #endif
+  #ifdef BINKD3D12GPUFUNCTIONS
+    #undef BINKD3D12GPUFUNCTIONS
+  #endif
+  #ifdef BINKNDAFUNCTIONS
+    #undef BINKNDAFUNCTIONS 
+  #endif
+  #ifdef BINKTEXTURESSUFFIX
+    #undef BINKTEXTURESSUFFIX 
+  #endif
+  #undef BINKTEXTURESH
+  #undef BINKTEXTURESCLEANUP
+#endif
+
+RADDEFEND
+
+#endif
diff --git a/Engine/Plugins/Media/BinkMedia/Source/SDK/include/binktiny.h b/Engine/Plugins/Media/BinkMedia/Source/SDK/include/binktiny.h
new file mode 100644
--- /dev/null
+++ b/Engine/Plugins/Media/BinkMedia/Source/SDK/include/binktiny.h
@@ -0,0 +1,89 @@
+#ifndef __BINKH__
+#define __BINKH__
+
+typedef struct BINK * HBINK;
+
+#define BINKALPHA              0x00100000L // Decompress alpha plane (if present)
+#define BINKHDR                0x00000004L // Video is an HDR video
+#define BINKUSETRIPLEBUFFERING 0x00000008L // Use triple buffering in the framebuffers
+
+typedef struct BINKRECT {
+  S32 Left,Top,Width,Height;
+} BINKRECT;
+
+#define BINKMAXDIRTYRECTS 1
+
+typedef struct BINKPLANE
+{
+  void * Buffer;
+  S32 Allocate;
+  U32 BufferPitch;
+} BINKPLANE;
+
+typedef struct BINKFRAMEPLANESET
+{
+  BINKPLANE YPlane;
+  BINKPLANE cRPlane;
+  BINKPLANE cBPlane;
+  BINKPLANE APlane;
+  BINKPLANE HPlane;
+} BINKFRAMEPLANESET;
+
+#define BINKMAXFRAMEBUFFERS 3
+
+typedef struct BINKFRAMEBUFFERS
+{
+  S32 TotalFrames;
+  U32 YABufferWidth;
+  U32 YABufferHeight;
+  U32 cRcBBufferWidth;
+  U32 cRcBBufferHeight;
+
+  U32 FrameNum;  // 0 to (TotalFrames-1)
+  BINKFRAMEPLANESET Frames[ BINKMAXFRAMEBUFFERS ];
+} BINKFRAMEBUFFERS;
+
+
+typedef struct BINK {
+  U32 Width;                  // Width (1 based, 640 for example)
+  U32 Height;                 // Height (1 based, 480 for example)
+  U32 Frames;                 // Number of frames (1 based, 100 = 100 frames)
+  U32 FrameNum;               // Frame to *be* displayed (1 based)
+
+  U32 LastFrameNum;           // Last frame decompressed or skipped (1 based)
+  U32 FrameRate;              // Frame Rate Numerator
+  U32 FrameRateDiv;           // Frame Rate Divisor (frame rate=numerator/divisor)
+  U32 ReadError;              // Non-zero if a read error has ocurred
+
+  U32 OpenFlags;              // flags used on open
+  U32 BinkType;               // Bink flags
+  U32 LargestFrameSize;       // Largest frame size
+  U32 FrameSize;              // The current frame's size in bytes
+
+  U32 SndSize;                // The current frame sound tracks' size in bytes
+  U32 FrameChangePercent;     // very rough percentage of the frame that changed
+  S32 NumTracks;              // how many tracks
+  S32 NumRects;
+
+  U32 soundon;                // sound turned on?
+  U32 videoon;                // video turned on?
+  U32 needio;                 // set to 1, if we need an io before the next binkframe
+  S32 closing;                // are we closing?
+
+  BINKRECT FrameRects[BINKMAXDIRTYRECTS];// Dirty rects from BinkGetRects
+
+  F32 ColorSpace[16];         // 4x4 matrix to use for colorspace conversion
+                              // in HDR, 0 = smpte2084(maxnits/10k)*0.5, 1/2=ct/cp scale, 3/4=ct/cp bias, 5 = maxnits
+
+  U64 FileOffset;             // Offset into the file where the Bink starts        
+  U64 Highest1SecRate;        // Highest 1 sec data rate
+
+  BINKFRAMEBUFFERS * FrameBuffers; // Bink frame buffers that we decompress to
+
+} BINK;
+
+RADEXPFUNC void RADEXPLINK BinkGetFrameBuffersInfo( HBINK bink, BINKFRAMEBUFFERS * fbset );
+RADEXPFUNC S32  RADEXPLINK BinkAllocateFrameBuffers( HBINK bp, BINKFRAMEBUFFERS * set, U32 minimum_alignment );
+RADEXPFUNC void RADEXPLINK BinkRegisterFrameBuffers( HBINK bink, BINKFRAMEBUFFERS * fbset );
+
+#endif
diff --git a/Engine/Plugins/Media/BinkMedia/Source/SDK/include/egttypes.h b/Engine/Plugins/Media/BinkMedia/Source/SDK/include/egttypes.h
--- a/Engine/Plugins/Media/BinkMedia/Source/SDK/include/egttypes.h
+++ b/Engine/Plugins/Media/BinkMedia/Source/SDK/include/egttypes.h
@@ -3,7 +3,7 @@
 //   and the separate 'Unreal Engine End User License Agreement for Publishing'.
 
 //===============================================
-// (C) Copyright 1994-2021 Epic Games Tools, LLC  
+// (C) Copyright 1994-2022 Epic Games Tools, LLC  
 //===============================================
 
 #ifndef __EGTTYPESH__
@@ -13,7 +13,7 @@
 #define __RADTYPESH__
 #define __RADRR_COREH__ // block old rr_core
 
-#define RADCOPYRIGHT "Copyright (C) 1994-2021, Epic Games Tools, LLC"
+#define RADCOPYRIGHT "Copyright (C) 1994-2022, Epic Games Tools, LLC"
 
 #if !defined(__RADRES__) // don't include anything for resource compiles
 
diff --git a/Engine/Source/Runtime/Core/Public/HAL/Platform.h b/Engine/Source/Runtime/Core/Public/HAL/Platform.h
--- a/Engine/Source/Runtime/Core/Public/HAL/Platform.h
+++ b/Engine/Source/Runtime/Core/Public/HAL/Platform.h
@@ -618,6 +618,14 @@
 	#define PLATFORM_USE_PLATFORM_FILE_MANAGED_STORAGE_WRAPPER	0
 #endif
 
+#ifndef PLATFORM_HAS_DIRECT_TEXTURE_MEMORY_ACCESS
+	#define PLATFORM_HAS_DIRECT_TEXTURE_MEMORY_ACCESS 0
+#endif
+
+#ifndef PLATFORM_DIRECT_TEXTURE_MEMORY_ACCESS_LOCK_MODE 
+	#define PLATFORM_DIRECT_TEXTURE_MEMORY_ACCESS_LOCK_MODE RLM_ReadOnly 
+#endif
+
 #ifndef PLATFORM_NEEDS_GPU_UAV_RESOURCE_INIT_WORKAROUND
 	#define PLATFORM_NEEDS_GPU_UAV_RESOURCE_INIT_WORKAROUND 0
 #endif
diff --git a/Engine/Source/Runtime/RenderCore/Private/BinkShaders.cpp b/Engine/Source/Runtime/RenderCore/Private/BinkShaders.cpp
new file mode 100644
--- /dev/null
+++ b/Engine/Source/Runtime/RenderCore/Private/BinkShaders.cpp
@@ -0,0 +1,9 @@
+// Copyright Epic Games, Inc. All Rights Reserved.
+
+#include "BinkShaders.h"
+#include "ShaderParameterUtils.h"
+#include "RHIStaticStates.h"
+
+IMPLEMENT_GLOBAL_SHADER(FBinkDrawVS, "/Engine/Private/Bink.usf", "BinkDrawVS", SF_Vertex);
+IMPLEMENT_GLOBAL_SHADER(FBinkDrawYCbCrPS, "/Engine/Private/Bink.usf", "BinkDrawYCbCr", SF_Pixel);
+IMPLEMENT_GLOBAL_SHADER(FBinkDrawICtCpPS, "/Engine/Private/Bink.usf", "BinkDrawICtCp", SF_Pixel);
diff --git a/Engine/Source/Runtime/RenderCore/Public/BinkShaders.h b/Engine/Source/Runtime/RenderCore/Public/BinkShaders.h
new file mode 100644
--- /dev/null
+++ b/Engine/Source/Runtime/RenderCore/Public/BinkShaders.h
@@ -0,0 +1,107 @@
+// Copyright Epic Games, Inc. All Rights Reserved.
+
+#pragma once
+
+#include "CoreMinimal.h"
+#include "RHI.h"
+#include "RenderResource.h"
+#include "Shader.h"
+#include "ShaderParameterStruct.h"
+#include "ShaderParameterUtils.h"
+#include "ShaderParameterMetadata.h"
+#include "GlobalShader.h"
+#include "PipelineStateCache.h"
+
+BEGIN_SHADER_PARAMETER_STRUCT(FBinkParameters, )
+	SHADER_PARAMETER_RDG_TEXTURE(Texture2D, tex0)
+	SHADER_PARAMETER_RDG_TEXTURE(Texture2D, tex1)
+	SHADER_PARAMETER_RDG_TEXTURE(Texture2D, tex2)
+	SHADER_PARAMETER_RDG_TEXTURE(Texture2D, tex3)
+	SHADER_PARAMETER_RDG_TEXTURE(Texture2D, tex4)
+	SHADER_PARAMETER_SAMPLER(SamplerState, samp0)
+	SHADER_PARAMETER_SAMPLER(SamplerState, samp1)
+	SHADER_PARAMETER_SAMPLER(SamplerState, samp2)
+	SHADER_PARAMETER_SAMPLER(SamplerState, samp3)
+	SHADER_PARAMETER_SAMPLER(SamplerState, samp4)
+	SHADER_PARAMETER(FVector4f, consta)
+	SHADER_PARAMETER(FVector4f, crc)
+	SHADER_PARAMETER(FVector4f, cbc)
+	SHADER_PARAMETER(FVector4f, adj)
+	SHADER_PARAMETER(FVector4f, yscale)
+	SHADER_PARAMETER(FVector4f, xy_xform0)
+	SHADER_PARAMETER(FVector4f, xy_xform1)
+	SHADER_PARAMETER(FVector4f, uv_xform0)
+	SHADER_PARAMETER(FVector4f, uv_xform1)
+	SHADER_PARAMETER(FVector4f, uv_xform2)
+	RENDER_TARGET_BINDING_SLOTS()
+END_SHADER_PARAMETER_STRUCT()
+
+struct FBinkDrawVS : FGlobalShader
+{
+	DECLARE_EXPORTED_SHADER_TYPE(FBinkDrawVS, Global, RENDERCORE_API);
+
+	SHADER_USE_PARAMETER_STRUCT(FBinkDrawVS, FGlobalShader);
+
+	static bool ShouldCompilePermutation(const FGlobalShaderPermutationParameters& Parameters) 
+	{
+		return true;
+	}
+
+	static void ModifyCompilationEnvironment(const FGlobalShaderPermutationParameters& Parameters, FShaderCompilerEnvironment& OutEnvironment)
+	{
+	}
+
+	BEGIN_SHADER_PARAMETER_STRUCT(FParameters, )
+		SHADER_PARAMETER_STRUCT_INCLUDE(FBinkParameters, BinkParameters)
+	END_SHADER_PARAMETER_STRUCT()
+};
+
+struct FBinkDrawYCbCrPS : FGlobalShader
+{
+	DECLARE_EXPORTED_SHADER_TYPE(FBinkDrawYCbCrPS, Global, RENDERCORE_API);
+	SHADER_USE_PARAMETER_STRUCT(FBinkDrawYCbCrPS, FGlobalShader);
+
+	class FALPHA    : SHADER_PERMUTATION_BOOL("ALPHA");
+	class FSRGB     : SHADER_PERMUTATION_BOOL("SRGB");
+
+	using FPermutationDomain = TShaderPermutationDomain<FALPHA, FSRGB>;
+
+	static bool ShouldCompilePermutation(const FGlobalShaderPermutationParameters& Parameters) 
+	{
+		return true;
+	}
+
+	static void ModifyCompilationEnvironment(const FGlobalShaderPermutationParameters& Parameters, FShaderCompilerEnvironment& OutEnvironment)
+	{
+	}
+
+	BEGIN_SHADER_PARAMETER_STRUCT(FParameters, )
+		SHADER_PARAMETER_STRUCT_INCLUDE(FBinkParameters, BinkParameters)
+	END_SHADER_PARAMETER_STRUCT()
+};
+
+struct FBinkDrawICtCpPS : FGlobalShader
+{
+	DECLARE_EXPORTED_SHADER_TYPE(FBinkDrawICtCpPS, Global, RENDERCORE_API);
+	SHADER_USE_PARAMETER_STRUCT(FBinkDrawICtCpPS, FGlobalShader);
+
+	class FALPHA    : SHADER_PERMUTATION_BOOL("ALPHA");
+	class FTONEMAP  : SHADER_PERMUTATION_BOOL("TONEMAP");
+	class FST2084   : SHADER_PERMUTATION_BOOL("ST2084");
+
+	using FPermutationDomain = TShaderPermutationDomain<FALPHA, FTONEMAP, FST2084>;
+
+	static bool ShouldCompilePermutation(const FGlobalShaderPermutationParameters& Parameters) 
+	{
+		return true;
+	}
+
+	static void ModifyCompilationEnvironment(const FGlobalShaderPermutationParameters& Parameters, FShaderCompilerEnvironment& OutEnvironment)
+	{
+	}
+
+	BEGIN_SHADER_PARAMETER_STRUCT(FParameters, )
+		SHADER_PARAMETER_STRUCT_INCLUDE(FBinkParameters, BinkParameters)
+	END_SHADER_PARAMETER_STRUCT()
+};
+
diff --git a/Engine/Source/Runtime/Windows/D3D11RHI/Private/D3D11Texture.cpp b/Engine/Source/Runtime/Windows/D3D11RHI/Private/D3D11Texture.cpp
--- a/Engine/Source/Runtime/Windows/D3D11RHI/Private/D3D11Texture.cpp
+++ b/Engine/Source/Runtime/Windows/D3D11RHI/Private/D3D11Texture.cpp
@@ -1510,11 +1510,11 @@ void FD3D11DynamicRHI::RHIUpdateTexture2D(FRHITexture2D* TextureRHI, uint32 MipI
 	};
 
 	check(UpdateRegion.Width % GPixelFormats[Texture->GetFormat()].BlockSizeX == 0);
-	check(UpdateRegion.Height % GPixelFormats[Texture->GetFormat()].BlockSizeX == 0);
+	check(UpdateRegion.Height % GPixelFormats[Texture->GetFormat()].BlockSizeY == 0);
 	check(UpdateRegion.DestX % GPixelFormats[Texture->GetFormat()].BlockSizeX == 0);
-	check(UpdateRegion.DestY % GPixelFormats[Texture->GetFormat()].BlockSizeX == 0);
+	check(UpdateRegion.DestY % GPixelFormats[Texture->GetFormat()].BlockSizeY == 0);
 	check(UpdateRegion.SrcX % GPixelFormats[Texture->GetFormat()].BlockSizeX == 0);
-	check(UpdateRegion.SrcY % GPixelFormats[Texture->GetFormat()].BlockSizeX == 0);
+	check(UpdateRegion.SrcY % GPixelFormats[Texture->GetFormat()].BlockSizeY == 0);
 
 	Direct3DDeviceIMContext->UpdateSubresource(Texture->GetResource(), MipIndex, &DestBox, SourceData, SourcePitch, 0);
 }
