Fix for linux editor weird linking issue (disables BinkMediaPlayerEditor on Linux Editor)

diff --git a/Engine/Plugins/Media/BinkMedia/BinkMedia.uplugin b/Engine/Plugins/Media/BinkMedia/BinkMedia.uplugin
--- a/Engine/Plugins/Media/BinkMedia/BinkMedia.uplugin
+++ b/Engine/Plugins/Media/BinkMedia/BinkMedia.uplugin
@@ -29,7 +29,10 @@
 		{
 			"Name" : "BinkMediaPlayerEditor",
 			"Type" : "Editor",
-			"TargetDenyList": [ "Server" ]
+			"TargetDenyList": [ "Server" ],
+			"PlatformDenyList": [ 
+				"Linux"
+			]
 		}
     ]
 }
