Replace TransitionPassRenderTargets with manual transitions.

diff --git a/Engine/Plugins/Media/BinkMedia/Source/BinkMediaPlayer/Private/Assets/BinkMediaTextureResource.cpp b/Engine/Plugins/Media/BinkMedia/Source/BinkMediaPlayer/Private/Assets/BinkMediaTextureResource.cpp
--- a/Engine/Plugins/Media/BinkMedia/Source/BinkMediaPlayer/Private/Assets/BinkMediaTextureResource.cpp
+++ b/Engine/Plugins/Media/BinkMedia/Source/BinkMediaPlayer/Private/Assets/BinkMediaTextureResource.cpp
@@ -68,11 +68,11 @@ void FBinkMediaTextureResource::InitDynamicRHI()
 	FRHICommandListImmediate& RHICmdList = FRHICommandListExecutor::GetImmediateCommandList();
 	{
 		FRHIRenderPassInfo RPInfo(TextureRHI, ERenderTargetActions::Clear_Store);
-		TransitionRenderPassTargets(RHICmdList, RPInfo);
+		RHICmdList.Transition(FRHITransitionInfo(TextureRHI.GetReference(), ERHIAccess::Unknown, ERHIAccess::RTV));
 		RHICmdList.BeginRenderPass(RPInfo, TEXT("ClearTexture"));
 		RHICmdList.EndRenderPass();
 		RHICmdList.SetViewport(0, 0, 0, w, h, 1);
-		RHICmdList.Transition(FRHITransitionInfo(TextureRHI.GetReference(), ERHIAccess::Unknown, ERHIAccess::UAVGraphics));
+		RHICmdList.Transition(FRHITransitionInfo(TextureRHI.GetReference(), ERHIAccess::RTV, ERHIAccess::UAVGraphics));
 	}
 }
 
@@ -116,11 +116,11 @@ void FBinkMediaTextureResource::Clear()
 	ENQUEUE_RENDER_COMMAND(BinkMediaPlayer_Draw)([ref,ref2,w,h](FRHICommandListImmediate& RHICmdList) 
 	{ 
 		FRHIRenderPassInfo RPInfo(ref2, ERenderTargetActions::Clear_Store);
-		TransitionRenderPassTargets(RHICmdList, RPInfo);
+		RHICmdList.Transition(FRHITransitionInfo(ref2.GetReference(), ERHIAccess::Unknown, ERHIAccess::RTV));
 		RHICmdList.BeginRenderPass(RPInfo, TEXT("ClearTexture"));
 		RHICmdList.EndRenderPass();
 		RHICmdList.SetViewport(0, 0, 0, w, h, 1);
-		RHICmdList.Transition(FRHITransitionInfo(ref2.GetReference(), ERHIAccess::Unknown, ERHIAccess::UAVGraphics));
+		RHICmdList.Transition(FRHITransitionInfo(ref2.GetReference(), ERHIAccess::RTV, ERHIAccess::UAVGraphics));
 	});
 }
 
diff --git a/Engine/Plugins/Runtime/nDisplay/Source/DisplayCluster/Private/Render/Device/DisplayClusterDeviceBase.cpp b/Engine/Plugins/Runtime/nDisplay/Source/DisplayCluster/Private/Render/Device/DisplayClusterDeviceBase.cpp
--- a/Engine/Plugins/Runtime/nDisplay/Source/DisplayCluster/Private/Render/Device/DisplayClusterDeviceBase.cpp
+++ b/Engine/Plugins/Runtime/nDisplay/Source/DisplayCluster/Private/Render/Device/DisplayClusterDeviceBase.cpp
@@ -494,7 +494,7 @@ void FDisplayClusterDeviceBase::RenderTexture_RenderThread(FRHICommandListImmedi
 		{
 			// Clear render target before out frame resolving, help to make things look better visually for console/resize, etc.
 			FRHIRenderPassInfo RPInfo(SrcTexture, ERenderTargetActions::Clear_Store);
-			TransitionRenderPassTargets(RHICmdList, RPInfo);
+			RHICmdList.Transition(FRHITransitionInfo(SrcTexture, ERHIAccess::Unknown, ERHIAccess::RTV));
 			RHICmdList.BeginRenderPass(RPInfo, TEXT("ClearTexture"));
 			RHICmdList.EndRenderPass();
 		}
diff --git a/Engine/Plugins/Runtime/nDisplay/Source/DisplayCluster/Private/Render/Viewport/DisplayClusterViewportManagerProxy.cpp b/Engine/Plugins/Runtime/nDisplay/Source/DisplayCluster/Private/Render/Viewport/DisplayClusterViewportManagerProxy.cpp
--- a/Engine/Plugins/Runtime/nDisplay/Source/DisplayCluster/Private/Render/Viewport/DisplayClusterViewportManagerProxy.cpp
+++ b/Engine/Plugins/Runtime/nDisplay/Source/DisplayCluster/Private/Render/Viewport/DisplayClusterViewportManagerProxy.cpp
@@ -272,7 +272,7 @@ void FDisplayClusterViewportManagerProxy::UpdateDeferredResources_RenderThread(F
 static void ImplClearRenderTargetResource_RenderThread(FRHICommandListImmediate& RHICmdList, FRHITexture2D* InRenderTargetTexture)
 {
 	FRHIRenderPassInfo RPInfo(InRenderTargetTexture, ERenderTargetActions::DontLoad_Store);
-	TransitionRenderPassTargets(RHICmdList, RPInfo);
+	RHICmdList.Transition(FRHITransitionInfo(InRenderTargetTexture, ERHIAccess::Unknown, ERHIAccess::RTV));
 	RHICmdList.BeginRenderPass(RPInfo, TEXT("nDisplay_ClearRTT"));
 	{
 		const FIntPoint Size = InRenderTargetTexture->GetSizeXY();
diff --git a/Engine/Plugins/Runtime/nDisplay/Source/DisplayCluster/Private/Render/Viewport/DisplayClusterViewportProxy.cpp b/Engine/Plugins/Runtime/nDisplay/Source/DisplayCluster/Private/Render/Viewport/DisplayClusterViewportProxy.cpp
--- a/Engine/Plugins/Runtime/nDisplay/Source/DisplayCluster/Private/Render/Viewport/DisplayClusterViewportProxy.cpp
+++ b/Engine/Plugins/Runtime/nDisplay/Source/DisplayCluster/Private/Render/Viewport/DisplayClusterViewportProxy.cpp
@@ -537,7 +537,7 @@ void ResampleCopyTextureImpl_RenderThread(FRHICommandListImmediate& RHICmdList,
 	// Texture format mismatch, use a shader to do the copy.
 	// #todo-renderpasses there's no explicit resolve here? Do we need one?
 	FRHIRenderPassInfo RPInfo(DstTexture, ERenderTargetActions::Load_Store);
-	TransitionRenderPassTargets(RHICmdList, RPInfo);
+	RHICmdList.Transition(FRHITransitionInfo(DstTexture, ERHIAccess::Unknown, ERHIAccess::RTV));
 
 	RHICmdList.BeginRenderPass(RPInfo, TEXT("DisaplyClusterRender_ResampleTexture"));
 	{
diff --git a/Engine/Plugins/Runtime/nDisplay/Source/DisplayClusterShaders/Private/Shaders/DisplayClusterShadersPostprocess_Blur.cpp b/Engine/Plugins/Runtime/nDisplay/Source/DisplayClusterShaders/Private/Shaders/DisplayClusterShadersPostprocess_Blur.cpp
--- a/Engine/Plugins/Runtime/nDisplay/Source/DisplayClusterShaders/Private/Shaders/DisplayClusterShadersPostprocess_Blur.cpp
+++ b/Engine/Plugins/Runtime/nDisplay/Source/DisplayClusterShaders/Private/Shaders/DisplayClusterShadersPostprocess_Blur.cpp
@@ -139,7 +139,7 @@ static void PicpBlurPostProcess_RenderThread(
 	TShaderMapRef<TPicpBlurPostProcessPS<ShaderType>> PixelShader(GlobalShaderMap);
 
 	FRHIRenderPassInfo RPInfo(OutRenderTargetableTexture, ERenderTargetActions::DontLoad_Store);
-	TransitionRenderPassTargets(RHICmdList, RPInfo);
+	RHICmdList.Transition(FRHITransitionInfo(OutRenderTargetableTexture, ERHIAccess::Unknown, ERHIAccess::RTV));
 
 	RHICmdList.BeginRenderPass(RPInfo, TEXT("nDisplay_PicpPostProcessBlur"));
 	{
diff --git a/Engine/Plugins/Runtime/nDisplay/Source/DisplayClusterShaders/Private/Shaders/DisplayClusterShadersPostprocess_OutputRemap.cpp b/Engine/Plugins/Runtime/nDisplay/Source/DisplayClusterShaders/Private/Shaders/DisplayClusterShadersPostprocess_OutputRemap.cpp
--- a/Engine/Plugins/Runtime/nDisplay/Source/DisplayClusterShaders/Private/Shaders/DisplayClusterShadersPostprocess_OutputRemap.cpp
+++ b/Engine/Plugins/Runtime/nDisplay/Source/DisplayClusterShaders/Private/Shaders/DisplayClusterShadersPostprocess_OutputRemap.cpp
@@ -141,7 +141,7 @@ bool FDisplayClusterShadersPostprocess_OutputRemap::RenderPostprocess_OutputRema
 	SCOPED_DRAW_EVENT(RHICmdList, nDisplay_PostProcess_OutputRemap);
 
 	FRHIRenderPassInfo RPInfo(InRenderTargetableDestTexture, ERenderTargetActions::Load_Store);
-	TransitionRenderPassTargets(RHICmdList, RPInfo);
+	RHICmdList.Transition(FRHITransitionInfo(InRenderTargetableDestTexture, ERHIAccess::Unknown, ERHIAccess::RTV));
 
 	RHICmdList.BeginRenderPass(RPInfo, TEXT("nDisplay_OutputRemap"));
 	{
diff --git a/Engine/Plugins/Runtime/nDisplay/Source/DisplayClusterShaders/Private/Shaders/DisplayClusterShadersWarpblend_ICVFX.cpp b/Engine/Plugins/Runtime/nDisplay/Source/DisplayClusterShaders/Private/Shaders/DisplayClusterShadersWarpblend_ICVFX.cpp
--- a/Engine/Plugins/Runtime/nDisplay/Source/DisplayClusterShaders/Private/Shaders/DisplayClusterShadersWarpblend_ICVFX.cpp
+++ b/Engine/Plugins/Runtime/nDisplay/Source/DisplayClusterShaders/Private/Shaders/DisplayClusterShadersWarpblend_ICVFX.cpp
@@ -930,7 +930,7 @@ bool FDisplayClusterShadersWarpblend_ICVFX::RenderWarpBlend_ICVFX(FRHICommandLis
 	// Do single warp render pass
 	bool bIsRenderSuccess = false;
 	FRHIRenderPassInfo RPInfo(InWarpBlendParameters.Dest.Texture, ERenderTargetActions::Load_Store);
-	TransitionRenderPassTargets(RHICmdList, RPInfo);
+	RHICmdList.Transition(FRHITransitionInfo(InWarpBlendParameters.Dest.Texture, ERHIAccess::Unknown, ERHIAccess::RTV));
 
 	RHICmdList.BeginRenderPass(RPInfo, TEXT("nDisplay_IcvfxWarpBlend"));
 	{
diff --git a/Engine/Plugins/Runtime/nDisplay/Source/DisplayClusterShaders/Private/Shaders/DisplayClusterShadersWarpblend_MPCDI.cpp b/Engine/Plugins/Runtime/nDisplay/Source/DisplayClusterShaders/Private/Shaders/DisplayClusterShadersWarpblend_MPCDI.cpp
--- a/Engine/Plugins/Runtime/nDisplay/Source/DisplayClusterShaders/Private/Shaders/DisplayClusterShadersWarpblend_MPCDI.cpp
+++ b/Engine/Plugins/Runtime/nDisplay/Source/DisplayClusterShaders/Private/Shaders/DisplayClusterShadersWarpblend_MPCDI.cpp
@@ -585,7 +585,7 @@ bool FDisplayClusterShadersWarpblend_MPCDI::RenderWarpBlend_MPCDI(FRHICommandLis
 	// Do single-pass warp&blend render
 	bool bIsRenderSuccess = false;
 	FRHIRenderPassInfo RPInfo(InWarpBlendParameters.Dest.Texture, ERenderTargetActions::Load_Store);
-	TransitionRenderPassTargets(RHICmdList, RPInfo);
+	RHICmdList.Transition(FRHITransitionInfo(InWarpBlendParameters.Dest.Texture, ERHIAccess::Unknown, ERHIAccess::RTV));
 
 	RHICmdList.BeginRenderPass(RPInfo, TEXT("nDisplay_MpcdiWarpBlend"));
 	{
diff --git a/Engine/Plugins/VirtualProduction/DMX/DMXPixelMapping/Source/DMXPixelMappingRenderer/Private/DMXPixelMappingRenderer.cpp b/Engine/Plugins/VirtualProduction/DMX/DMXPixelMapping/Source/DMXPixelMappingRenderer/Private/DMXPixelMappingRenderer.cpp
--- a/Engine/Plugins/VirtualProduction/DMX/DMXPixelMapping/Source/DMXPixelMappingRenderer/Private/DMXPixelMappingRenderer.cpp
+++ b/Engine/Plugins/VirtualProduction/DMX/DMXPixelMapping/Source/DMXPixelMappingRenderer/Private/DMXPixelMappingRenderer.cpp
@@ -259,7 +259,7 @@ void FDMXPixelMappingRenderer::RenderPreview(const FTextureResource* TextureReso
 		// Clear preview texture
 		{
 			FRHIRenderPassInfo RPInfo(RenderTargetRef, ERenderTargetActions::DontLoad_Store);
-			TransitionRenderPassTargets(RHICmdList, RPInfo);
+			RHICmdList.Transition(FRHITransitionInfo(RenderTargetRef, ERHIAccess::Unknown, ERHIAccess::RTV));
 			RHICmdList.BeginRenderPass(RPInfo, TEXT("ClearCanvas"));
 			RHICmdList.SetViewport(0.f, 0.f, 0.f, OutputTextureSize.X, OutputTextureSize.Y, 1.f);
 			DrawClearQuad(RHICmdList, FLinearColor::Black);
diff --git a/Engine/Source/Runtime/Engine/Private/KismetRenderingLibrary.cpp b/Engine/Source/Runtime/Engine/Private/KismetRenderingLibrary.cpp
--- a/Engine/Source/Runtime/Engine/Private/KismetRenderingLibrary.cpp
+++ b/Engine/Source/Runtime/Engine/Private/KismetRenderingLibrary.cpp
@@ -54,7 +54,7 @@ void UKismetRenderingLibrary::ClearRenderTarget2D(UObject* WorldContextObject, U
 			[RenderTargetResource, ClearColor](FRHICommandList& RHICmdList)
 		{
 			FRHIRenderPassInfo RPInfo(RenderTargetResource->GetRenderTargetTexture(), ERenderTargetActions::DontLoad_Store);
-			TransitionRenderPassTargets(RHICmdList, RPInfo);
+			RHICmdList.Transition(FRHITransitionInfo(RenderTargetResource->GetRenderTargetTexture(), ERHIAccess::Unknown, ERHIAccess::RTV));
 			RHICmdList.BeginRenderPass(RPInfo, TEXT("ClearRT"));
 			DrawClearQuad(RHICmdList, ClearColor);
 			RHICmdList.EndRenderPass();
diff --git a/Engine/Source/Runtime/Engine/Private/Particles/ParticleCurveTexture.cpp b/Engine/Source/Runtime/Engine/Private/Particles/ParticleCurveTexture.cpp
--- a/Engine/Source/Runtime/Engine/Private/Particles/ParticleCurveTexture.cpp
+++ b/Engine/Source/Runtime/Engine/Private/Particles/ParticleCurveTexture.cpp
@@ -189,7 +189,7 @@ static void InjectCurves(
 	}
 
 	FRHIRenderPassInfo RPInfo(CurveTextureTargetRHI, MakeRenderTargetActions(LoadAction, ERenderTargetStoreAction::EStore));
-	TransitionRenderPassTargets(RHICmdList, RPInfo);
+	RHICmdList.Transition(FRHITransitionInfo(CurveTextureTargetRHI, ERHIAccess::Unknown, ERHIAccess::RTV));
 	RHICmdList.BeginRenderPass(RPInfo, TEXT("InjectCurves"));
 	{
 		FGraphicsPipelineStateInitializer GraphicsPSOInit;
diff --git a/Engine/Source/Runtime/Engine/Private/Particles/ParticleGpuSimulation.cpp b/Engine/Source/Runtime/Engine/Private/Particles/ParticleGpuSimulation.cpp
--- a/Engine/Source/Runtime/Engine/Private/Particles/ParticleGpuSimulation.cpp
+++ b/Engine/Source/Runtime/Engine/Private/Particles/ParticleGpuSimulation.cpp
@@ -5185,7 +5185,10 @@ void FFXSystem::UpdateMultiGPUResources(FRHICommandListImmediate& RHICmdList)
 		RPInfo.ColorRenderTargets[0].ResolveTarget = ParticleSimulationResources->RenderAttributesTexture.TextureRHI;
 		RPInfo.ColorRenderTargets[1].ResolveTarget = ParticleSimulationResources->SimulationAttributesTexture.TextureRHI;
 		{
-			TransitionRenderPassTargets(RHICmdList, RPInfo);
+			RHICmdList.Transition({
+				FRHITransitionInfo(ParticleSimulationResources->RenderAttributesTexture.TextureTargetRHI, ERHIAccess::Unknown, ERHIAccess::RTV),
+				FRHITransitionInfo(ParticleSimulationResources->SimulationAttributesTexture.TextureTargetRHI, ERHIAccess::Unknown, ERHIAccess::RTV)
+			});
 			RHICmdList.BeginRenderPass(RPInfo, TEXT("UpdateMultiGPUResources"));
 
 			RHICmdList.SetViewport(0.0f, 0.0f, 0.0f, (float)GParticleSimulationTextureSizeX, (float)GParticleSimulationTextureSizeY, 1.0f);
diff --git a/Engine/Source/Runtime/Engine/Private/TextureRenderTarget2DArray.cpp b/Engine/Source/Runtime/Engine/Private/TextureRenderTarget2DArray.cpp
--- a/Engine/Source/Runtime/Engine/Private/TextureRenderTarget2DArray.cpp
+++ b/Engine/Source/Runtime/Engine/Private/TextureRenderTarget2DArray.cpp
@@ -307,7 +307,7 @@ void FTextureRenderTarget2DArrayResource::UpdateDeferredResource(FRHICommandList
 	const ERenderTargetLoadAction LoadAction = bClearRenderTarget ? ERenderTargetLoadAction::EClear : ERenderTargetLoadAction::ELoad;
 
 	FRHIRenderPassInfo RPInfo(RenderTarget2DArrayRHI, MakeRenderTargetActions(LoadAction, ERenderTargetStoreAction::EStore));
-	TransitionRenderPassTargets(RHICmdList, RPInfo);
+	RHICmdList.Transition(FRHITransitionInfo(RenderTarget2DArrayRHI, ERHIAccess::Unknown, ERHIAccess::RTV));
 	RHICmdList.BeginRenderPass(RPInfo, TEXT("UpdateTarget2DArray"));
 	RHICmdList.SetViewport(0.0f, 0.0f, 0.0f, (float)Dims.X, (float)Dims.Y, 1.0f);
 	RHICmdList.EndRenderPass();
diff --git a/Engine/Source/Runtime/Engine/Private/TextureRenderTargetCube.cpp b/Engine/Source/Runtime/Engine/Private/TextureRenderTargetCube.cpp
--- a/Engine/Source/Runtime/Engine/Private/TextureRenderTargetCube.cpp
+++ b/Engine/Source/Runtime/Engine/Private/TextureRenderTargetCube.cpp
@@ -330,19 +330,13 @@ void FTextureRenderTargetCubeResource::UpdateDeferredResource(FRHICommandListImm
 		}
 		
 		FRHIRenderPassInfo RPInfo(RenderTargetTextureRHI, MakeRenderTargetActions(LoadAction, ERenderTargetStoreAction::EStore));
-		TransitionRenderPassTargets(RHICmdList, RPInfo);
+		RHICmdList.Transition(FRHITransitionInfo(RenderTargetTextureRHI, ERHIAccess::Unknown, ERHIAccess::RTV));
 		RHICmdList.BeginRenderPass(RPInfo, TEXT("UpdateTargetCube"));
 		RHICmdList.SetViewport(0.0f, 0.0f, 0.0f, (float)Dims.X, (float)Dims.Y, 1.0f);
 		RHICmdList.EndRenderPass();
 		// copy surface to the texture for use
 		FResolveParams ResolveParams;
 		ResolveParams.CubeFace = (ECubeFace)FaceIdx;
-
-		FRHITransitionInfo TransitionsBefore[] = {
-			FRHITransitionInfo(RenderTargetTextureRHI, ERHIAccess::RTV, ERHIAccess::ResolveSrc),
-			FRHITransitionInfo(TextureCubeRHI, ERHIAccess::Unknown, ERHIAccess::ResolveDst)
-		};
-		RHICmdList.Transition(MakeArrayView(TransitionsBefore, UE_ARRAY_COUNT(TransitionsBefore)));
 		RHICmdList.CopyToResolveTarget(RenderTargetTextureRHI, TextureCubeRHI, ResolveParams);
 	}
 }
diff --git a/Engine/Source/Runtime/Engine/Private/TextureRenderTargetVolume.cpp b/Engine/Source/Runtime/Engine/Private/TextureRenderTargetVolume.cpp
--- a/Engine/Source/Runtime/Engine/Private/TextureRenderTargetVolume.cpp
+++ b/Engine/Source/Runtime/Engine/Private/TextureRenderTargetVolume.cpp
@@ -314,7 +314,7 @@ void FTextureRenderTargetVolumeResource::UpdateDeferredResource(FRHICommandListI
 	const ERenderTargetLoadAction LoadAction = bClearRenderTarget ? ERenderTargetLoadAction::EClear : ERenderTargetLoadAction::ELoad;
 
 	FRHIRenderPassInfo RPInfo(RenderTargetVolumeRHI, MakeRenderTargetActions(LoadAction, ERenderTargetStoreAction::EStore));
-	TransitionRenderPassTargets(RHICmdList, RPInfo);
+	RHICmdList.Transition(FRHITransitionInfo(RenderTargetVolumeRHI, ERHIAccess::Unknown, ERHIAccess::RTV));
 	RHICmdList.BeginRenderPass(RPInfo, TEXT("UpdateTargetVolume"));
 	RHICmdList.SetViewport(0.0f, 0.0f, 0.0f, (float)Dims.X, (float)Dims.Y, 1.0f);
 	RHICmdList.EndRenderPass();
diff --git a/Engine/Source/Runtime/Engine/Private/UserInterface/Canvas.cpp b/Engine/Source/Runtime/Engine/Private/UserInterface/Canvas.cpp
--- a/Engine/Source/Runtime/Engine/Private/UserInterface/Canvas.cpp
+++ b/Engine/Source/Runtime/Engine/Private/UserInterface/Canvas.cpp
@@ -1043,14 +1043,14 @@ void FCanvas::Clear(const FLinearColor& ClearColor)
 				{
 					// do fast clear
 					FRHIRenderPassInfo RPInfo(CanvasRenderTarget->GetRenderTargetTexture(), ERenderTargetActions::Clear_Store);
-					TransitionRenderPassTargets(RHICmdList, RPInfo);
+					RHICmdList.Transition(FRHITransitionInfo(CanvasRenderTarget->GetRenderTargetTexture(), ERHIAccess::Unknown, ERHIAccess::RTV));
 					RHICmdList.BeginRenderPass(RPInfo, TEXT("ClearCanvas"));
 					RHICmdList.EndRenderPass();					
 				}
 				else if(CanvasRenderTarget->GetRenderTargetTexture())
 				{
 					FRHIRenderPassInfo RPInfo(CanvasRenderTarget->GetRenderTargetTexture(), ERenderTargetActions::DontLoad_Store);
-					TransitionRenderPassTargets(RHICmdList, RPInfo);
+					RHICmdList.Transition(FRHITransitionInfo(CanvasRenderTarget->GetRenderTargetTexture(), ERHIAccess::Unknown, ERHIAccess::RTV));
 					RHICmdList.BeginRenderPass(RPInfo, TEXT("ClearCanvas"));
 					RHICmdList.SetViewport(0, 0, 0.0f, CanvasRenderTarget->GetSizeXY().X, CanvasRenderTarget->GetSizeXY().Y, 1.0f);
 					DrawClearQuad(RHICmdList, ClearColor);
diff --git a/Engine/Source/Runtime/RenderCore/Private/LongGPUTask.cpp b/Engine/Source/Runtime/RenderCore/Private/LongGPUTask.cpp
--- a/Engine/Source/Runtime/RenderCore/Private/LongGPUTask.cpp
+++ b/Engine/Source/Runtime/RenderCore/Private/LongGPUTask.cpp
@@ -17,7 +17,7 @@ void IssueScalableLongGPUTask(FRHICommandListImmediate& RHICmdList, int32 NumIte
 	FTexture2DRHIRef LongTaskRenderTarget = RHICreateTexture2D(1920, 1080, PF_B8G8R8A8, 1, 1, TexCreate_RenderTargetable, Info);
 
 	FRHIRenderPassInfo RPInfo(LongTaskRenderTarget, ERenderTargetActions::DontLoad_Store);
-	TransitionRenderPassTargets(RHICmdList, RPInfo);
+	RHICmdList.Transition(FRHITransitionInfo(LongTaskRenderTarget, ERHIAccess::Unknown, ERHIAccess::RTV));
 	RHICmdList.BeginRenderPass(RPInfo, TEXT("LongGPUTask"));
 
 	{
diff --git a/Engine/Source/Runtime/Renderer/Private/GPUBenchmark.cpp b/Engine/Source/Runtime/Renderer/Private/GPUBenchmark.cpp
--- a/Engine/Source/Runtime/Renderer/Private/GPUBenchmark.cpp
+++ b/Engine/Source/Runtime/Renderer/Private/GPUBenchmark.cpp
@@ -560,7 +560,7 @@ void RendererGPUBenchmark(FRHICommandListImmediate& RHICmdList, FSynthBenchmarkR
 				LocalWorkScale[Iteration] = (Iteration / 10.f + 1.f) * WorkScale;
 
 				FRHIRenderPassInfo RPInfo(RTItems[DestRTIndex]->GetRHI(), ERenderTargetActions::Load_Store);
-				TransitionRenderPassTargets(RHICmdList, RPInfo);
+				RHICmdList.Transition(FRHITransitionInfo(RTItems[DestRTIndex]->GetRHI(), ERHIAccess::Unknown, ERHIAccess::RTV));
 				RHICmdList.BeginRenderPass(RPInfo, TEXT("GPUBenchmark"));
 				{
 					RunBenchmarkShader(RHICmdList, VertexBuffer, View, MethodId, RTItems[SrcRTIndex], LocalWorkScale[Iteration]);
diff --git a/Engine/Source/Runtime/SlateRHIRenderer/Private/SlateRHIRenderer.cpp b/Engine/Source/Runtime/SlateRHIRenderer/Private/SlateRHIRenderer.cpp
--- a/Engine/Source/Runtime/SlateRHIRenderer/Private/SlateRHIRenderer.cpp
+++ b/Engine/Source/Runtime/SlateRHIRenderer/Private/SlateRHIRenderer.cpp
@@ -993,7 +993,7 @@ void FSlateRHIRenderer::DrawWindow_RenderThread(FRHICommandListImmediate& RHICmd
 				{
 					// #todo-renderpasses will this touch every pixel? use NoAction?
 					FRHIRenderPassInfo RPInfo(ViewportInfo.ColorSpaceLUTRT, ERenderTargetActions::Load_Store);
-					TransitionRenderPassTargets(RHICmdList, RPInfo);
+					RHICmdList.Transition(FRHITransitionInfo(ViewportInfo.ColorSpaceLUTRT, ERHIAccess::Unknown, ERHIAccess::RTV));
 					RHICmdList.BeginRenderPass(RPInfo, TEXT("GenerateLUT"));
 					{
 						FGraphicsPipelineStateInitializer GraphicsPSOInit;
@@ -1110,7 +1110,7 @@ void FSlateRHIRenderer::DrawWindow_RenderThread(FRHICommandListImmediate& RHICmd
 				auto ShaderMap = GetGlobalShaderMap(FeatureLevel);
 
 				FRHIRenderPassInfo RPInfo(FinalBuffer, ERenderTargetActions::Load_Store);
-				TransitionRenderPassTargets(RHICmdList, RPInfo);
+				RHICmdList.Transition(FRHITransitionInfo(FinalBuffer, ERHIAccess::Unknown, ERHIAccess::RTV));
 				RHICmdList.BeginRenderPass(RPInfo, TEXT("SlateComposite"));
 
 				FGraphicsPipelineStateInitializer GraphicsPSOInit;
