Moved Bink SDK into a new BinkMediaSDK module so it can be privately depended on.
#p4v-preflight-copy 20044209

diff --git a/Engine/Plugins/Media/BinkMedia/Source/BinkMediaPlayer/BinkMediaPlayer.Build.cs b/Engine/Plugins/Media/BinkMedia/Source/BinkMediaPlayer/BinkMediaPlayer.Build.cs
--- a/Engine/Plugins/Media/BinkMedia/Source/BinkMediaPlayer/BinkMediaPlayer.Build.cs
+++ b/Engine/Plugins/Media/BinkMedia/Source/BinkMediaPlayer/BinkMediaPlayer.Build.cs
@@ -7,15 +7,6 @@ using UnrealBuildTool;
 
 public class BinkMediaPlayer : ModuleRules 
 {
-	// Platform Extensions need to override these
-	protected virtual string LibRootDirectory { get { return ModuleDirectory; } }
-	protected virtual string LibName { get { return null; } }
-
-	protected virtual string SdkBaseDirectory { get { return Path.Combine(LibRootDirectory, "..", "SDK"); } }
-	protected virtual string LibDirectory { get { return Path.Combine(SdkBaseDirectory, "lib"); } }
-
-	protected virtual string IncDirectory { get { return Path.Combine(ModuleDirectory, "..", "SDK", "include"); } }
-
     public BinkMediaPlayer(ReadOnlyTargetRules Target) : base(Target)
     {
 		bAllowConfidentialPlatformDefines = true;
@@ -30,6 +21,8 @@ public class BinkMediaPlayer : ModuleRules
         //PublicDependencyModuleNames.Add("MediaAssets");
         PublicDependencyModuleNames.Add("Projects");
 
+		AddEngineThirdPartyPrivateStaticDependencies(Target, "BinkMediaPlayerSDK");
+
         PrivatePCHHeaderFile = "Private/BinkMediaPlayerPCH.h";
 
         if (Target.bBuildEditor == true)
@@ -45,56 +38,12 @@ public class BinkMediaPlayer : ModuleRules
             PublicDefinitions.Add("BINKPLUGIN_UE4_EDITOR=0");
         }
 
-		PublicDefinitions.Add("BUILDING_FOR_UNREAL_ONLY=1");
-		PublicDefinitions.Add("__RADNOEXPORTS__=1");
-		PublicDefinitions.Add("__RADINSTATICLIB__=1");
 		RuntimeDependencies.Add("$(ProjectDir)/Content/Movies/..."); // For chunked streaming
 
-		string Lib = LibName;
-		string Platform = Target.Platform.ToString();
-
-		if(Lib == null)
-		{
-			if (Target.IsInPlatformGroup(UnrealPlatformGroup.Microsoft))
-			{
-				Lib = "BinkUnreal" + Platform + ".lib";
-			}
-			else if (Target.Platform == UnrealTargetPlatform.Mac)
-			{
-				Lib = "BinkUnreal" + Platform + ".a";
-				PublicDependencyModuleNames.Add("MetalRHI");
-				PublicFrameworks.Add("AudioToolbox");
-			}
-            else if (Target.Platform == UnrealTargetPlatform.IOS)
-            {
-                Lib = "libBinkUnreal" + Platform + ".a";
-                PublicDependencyModuleNames.Add("MetalRHI");
-            }
-            else if (Target.IsInPlatformGroup(UnrealPlatformGroup.Apple))
-			{
-				Lib = "BinkUnreal" + Platform + ".a";
-				PublicDependencyModuleNames.Add("MetalRHI");
-			}
-			else if (Target.Platform == UnrealTargetPlatform.Android)
-			{
-				PublicDependencyModuleNames.Add("Launch");
-				PublicAdditionalLibraries.Add(Path.Combine(LibDirectory, "android", "armv7", "libBinkUnrealAndroid.a"));
-				PublicAdditionalLibraries.Add(Path.Combine(LibDirectory, "android", "arm64", "libBinkUnrealAndroid.a"));
-				PublicAdditionalLibraries.Add(Path.Combine(LibDirectory, "android", "x86", "libBinkUnrealAndroid.a"));
-				PublicAdditionalLibraries.Add(Path.Combine(LibDirectory, "android", "x64", "libBinkUnrealAndroid.a"));
-				string PluginPath = Utils.MakePathRelativeTo(ModuleDirectory, Target.RelativeEnginePath);
-				AdditionalPropertiesForReceipt.Add("AndroidPlugin", Path.Combine(PluginPath, "BinkMediaPlayer_APL.xml"));
-			}
-			else
-			{
-				Lib = "BinkUnreal" + Platform + ".a";
-			}
-		}
-
-		if (Lib != null)
+		if (Target.Platform == UnrealTargetPlatform.Android)
 		{
-			PublicAdditionalLibraries.Add(Path.Combine(LibDirectory, Lib));
+			string PluginPath = Utils.MakePathRelativeTo(ModuleDirectory, Target.RelativeEnginePath);
+			AdditionalPropertiesForReceipt.Add("AndroidPlugin", Path.Combine(PluginPath, "BinkMediaPlayer_APL.xml"));
 		}
-        PublicIncludePaths.Add(IncDirectory);
 	}
 }
diff --git a/Engine/Plugins/Media/BinkMedia/Source/BinkMediaPlayer/Private/Assets/BinkMediaPlayer.cpp b/Engine/Plugins/Media/BinkMedia/Source/BinkMediaPlayer/Private/Assets/BinkMediaPlayer.cpp
--- a/Engine/Plugins/Media/BinkMedia/Source/BinkMediaPlayer/Private/Assets/BinkMediaPlayer.cpp
+++ b/Engine/Plugins/Media/BinkMedia/Source/BinkMediaPlayer/Private/Assets/BinkMediaPlayer.cpp
@@ -2,6 +2,7 @@
 //   Licenced under the Unreal Engine EULA 
 
 #include "BinkMediaPlayer.h"
+#include "binkplugin.h"
 #include "BinkMediaPlayerPCH.h"
 #include "BinkMediaTexture.h"
 #include "BinkFunctionLibrary.h"
@@ -548,3 +549,25 @@ void UBinkMediaPlayer::Draw(UTexture *texture, bool tonemap, int out_nits, float
 		parms.player->UpdateTexture(RHICmdList, parms.ref, parms.native, parms.width, parms.height, false, parms.tonemap, parms.out_nits, parms.alpha, parms.srgb_decode, parms.hdr);
 	});
 }
+
+FIntPoint UBinkMediaPlayer::GetDimensions() const
+{
+	if (bnk)
+	{
+		BINKPLUGININFO bpinfo = {};
+		BinkPluginInfo(bnk, &bpinfo);
+		return FIntPoint(bpinfo.Width, bpinfo.Height);
+	}
+	return FIntPoint(0, 0);
+}
+
+float UBinkMediaPlayer::GetFrameRate() const
+{
+	if (bnk)
+	{
+		BINKPLUGININFO bpinfo = {};
+		BinkPluginInfo(bnk, &bpinfo);
+		return (float)(((double)bpinfo.FrameRate) / ((double)bpinfo.FrameRateDiv));
+	}
+	return 0;
+}
diff --git a/Engine/Plugins/Media/BinkMedia/Source/BinkMediaPlayer/Private/binktexturesRHI.cpp b/Engine/Plugins/Media/BinkMedia/Source/BinkMediaPlayer/Private/binktexturesRHI.cpp
--- a/Engine/Plugins/Media/BinkMedia/Source/BinkMediaPlayer/Private/binktexturesRHI.cpp
+++ b/Engine/Plugins/Media/BinkMedia/Source/BinkMediaPlayer/Private/binktexturesRHI.cpp
@@ -1,9 +1,9 @@
 // Copyright Epic Games, Inc. All Rights Reserved.
 
 #define BINKRHIFUNCTIONS
-#include "../../SDK/include/egttypes.h"
-#include "../../SDK/include/binktiny.h"
-#include "../../SDK/include/binktextures.h"
+#include "../../BinkMediaPlayerSDK/include/egttypes.h"
+#include "../../BinkMediaPlayerSDK/include/binktiny.h"
+#include "../../BinkMediaPlayerSDK/include/binktextures.h"
 
 #include "CoreMinimal.h"
 #include "UObject/ConstructorHelpers.h"
diff --git a/Engine/Plugins/Media/BinkMedia/Source/BinkMediaPlayer/Public/BinkMediaPlayer.h b/Engine/Plugins/Media/BinkMedia/Source/BinkMediaPlayer/Public/BinkMediaPlayer.h
--- a/Engine/Plugins/Media/BinkMedia/Source/BinkMediaPlayer/Public/BinkMediaPlayer.h
+++ b/Engine/Plugins/Media/BinkMedia/Source/BinkMediaPlayer/Public/BinkMediaPlayer.h
@@ -6,7 +6,6 @@
 #include "Tickable.h"
 #include "Delegates/DelegateCombinations.h"
 #include "RHI.h"
-#include "binkplugin.h"
 #include "BinkMediaPlayer.generated.h"
 
 extern BINKMEDIAPLAYER_API unsigned bink_gpu_api;
@@ -463,29 +462,10 @@ public:
 	bool IsReady() const { return bnk != NULL; }
 	void UpdateTexture(FRHICommandListImmediate &RHICmdList, FTexture2DRHIRef ref, void *nativePtr, int width, int height, bool isEditor, bool tonemap, int output_nits, float alpha, bool srgb_decode, bool is_hdr);
 
-	FIntPoint GetDimensions() const 
-	{
-		if (bnk) 
-		{
-			BINKPLUGININFO bpinfo = {};
-			BinkPluginInfo(bnk, &bpinfo);
-			return FIntPoint(bpinfo.Width, bpinfo.Height);
-		}
-		return FIntPoint(0, 0);
-	}
-
-	float GetFrameRate() const 
-	{
-		if (bnk) 
-		{
-			BINKPLUGININFO bpinfo = {};
-			BinkPluginInfo(bnk, &bpinfo);
-			return (float)(((double)bpinfo.FrameRate) / ((double)bpinfo.FrameRateDiv));
-		}
-		return 0;
-	}
-
-	BINKPLUGIN *bnk;
+	FIntPoint GetDimensions() const;
+	float GetFrameRate() const;
+
+	struct BINKPLUGIN *bnk;
 	bool paused;
 	bool reached_end;
 	FDelegateHandle overlayHook;
diff --git a/Engine/Plugins/Media/BinkMedia/Source/BinkMediaPlayerEditor/BinkMediaPlayerEditor.Build.cs b/Engine/Plugins/Media/BinkMedia/Source/BinkMediaPlayerEditor/BinkMediaPlayerEditor.Build.cs
--- a/Engine/Plugins/Media/BinkMedia/Source/BinkMediaPlayerEditor/BinkMediaPlayerEditor.Build.cs
+++ b/Engine/Plugins/Media/BinkMedia/Source/BinkMediaPlayerEditor/BinkMediaPlayerEditor.Build.cs
@@ -42,6 +42,8 @@ public class BinkMediaPlayerEditor : ModuleRules
         PrivatePCHHeaderFile = "Private/BinkMediaPlayerEditorPCH.h";
 
 		PublicDefinitions.Add("BUILDING_FOR_UNREAL_ONLY=1");
+		PublicDefinitions.Add("__RADNOEXPORTS__=1");
+		PublicDefinitions.Add("__RADINSTATICLIB__=1");
 		PublicDependencyModuleNames.Add("BinkMediaPlayer");
 
 		PrivateIncludePathModuleNames.AddRange(
diff --git a/Engine/Plugins/Media/BinkMedia/Source/BinkMediaPlayerSDK/BinkMediaPlayerSDK.Build.cs b/Engine/Plugins/Media/BinkMedia/Source/BinkMediaPlayerSDK/BinkMediaPlayerSDK.Build.cs
new file mode 100644
--- /dev/null
+++ b/Engine/Plugins/Media/BinkMedia/Source/BinkMediaPlayerSDK/BinkMediaPlayerSDK.Build.cs
@@ -0,0 +1,100 @@
+// Copyright Epic Games Tools LLC
+//   Licenced under the Unreal Engine EULA 
+
+using System.IO;
+using System.Collections.Generic;
+using UnrealBuildTool;
+
+public class BinkMediaPlayerSDK : ModuleRules 
+{
+	// Platform Extensions need to override these
+	protected virtual string LibRootDirectory { get { return ModuleDirectory; } }
+	protected virtual string LibName { get { return null; } }
+
+	protected virtual string SdkBaseDirectory { get { return LibRootDirectory; } }
+	protected virtual string LibDirectory { get { return Path.Combine(SdkBaseDirectory, "lib"); } }
+
+	protected virtual string IncDirectory { get { return Path.Combine(ModuleDirectory, "include"); } }
+
+    public BinkMediaPlayerSDK(ReadOnlyTargetRules Target) : base(Target)
+    {
+		Type = ModuleType.External;
+
+		bAllowConfidentialPlatformDefines = true;
+
+        PublicDependencyModuleNames.Add("Core");
+        PublicDependencyModuleNames.Add("CoreUObject");
+        PublicDependencyModuleNames.Add("Engine");
+        PublicDependencyModuleNames.Add("InputCore");
+        PublicDependencyModuleNames.Add("RenderCore");
+        PublicDependencyModuleNames.Add("RHI");
+        PublicDependencyModuleNames.Add("MoviePlayer");
+        //PublicDependencyModuleNames.Add("MediaAssets");
+        PublicDependencyModuleNames.Add("Projects");
+
+        PrivatePCHHeaderFile = "Private/BinkMediaPlayerPCH.h";
+
+        if (Target.bBuildEditor == true)
+        {
+			PublicDependencyModuleNames.Add("Slate");
+			PublicDependencyModuleNames.Add("SlateCore");
+			PublicDependencyModuleNames.Add("DesktopWidgets");
+			PublicDefinitions.Add("BINKPLUGIN_UE4_EDITOR=1");
+            PrivateDependencyModuleNames.AddRange(new string[] { "PropertyEditor", "DesktopPlatform", "SourceControl", "EditorStyle", "UnrealEd" });
+        }
+        else
+        {
+            PublicDefinitions.Add("BINKPLUGIN_UE4_EDITOR=0");
+        }
+
+		PublicDefinitions.Add("BUILDING_FOR_UNREAL_ONLY=1");
+		PublicDefinitions.Add("__RADNOEXPORTS__=1");
+		PublicDefinitions.Add("__RADINSTATICLIB__=1");
+		RuntimeDependencies.Add("$(ProjectDir)/Content/Movies/..."); // For chunked streaming
+
+		string Lib = LibName;
+		string Platform = Target.Platform.ToString();
+
+		if(Lib == null)
+		{
+			if (Target.IsInPlatformGroup(UnrealPlatformGroup.Microsoft))
+			{
+				Lib = "BinkUnreal" + Platform + ".lib";
+			}
+			else if (Target.Platform == UnrealTargetPlatform.Mac)
+			{
+				Lib = "BinkUnreal" + Platform + ".a";
+				PublicDependencyModuleNames.Add("MetalRHI");
+				PublicFrameworks.Add("AudioToolbox");
+			}
+            else if (Target.Platform == UnrealTargetPlatform.IOS)
+            {
+                Lib = "libBinkUnreal" + Platform + ".a";
+                PublicDependencyModuleNames.Add("MetalRHI");
+            }
+            else if (Target.IsInPlatformGroup(UnrealPlatformGroup.Apple))
+			{
+				Lib = "BinkUnreal" + Platform + ".a";
+				PublicDependencyModuleNames.Add("MetalRHI");
+			}
+			else if (Target.Platform == UnrealTargetPlatform.Android)
+			{
+				PublicDependencyModuleNames.Add("Launch");
+				PublicAdditionalLibraries.Add(Path.Combine(LibDirectory, "android", "armv7", "libBinkUnrealAndroid.a"));
+				PublicAdditionalLibraries.Add(Path.Combine(LibDirectory, "android", "arm64", "libBinkUnrealAndroid.a"));
+				PublicAdditionalLibraries.Add(Path.Combine(LibDirectory, "android", "x86", "libBinkUnrealAndroid.a"));
+				PublicAdditionalLibraries.Add(Path.Combine(LibDirectory, "android", "x64", "libBinkUnrealAndroid.a"));
+			}
+			else
+			{
+				Lib = "BinkUnreal" + Platform + ".a";
+			}
+		}
+
+		if (Lib != null)
+		{
+			PublicAdditionalLibraries.Add(Path.Combine(LibDirectory, Lib));
+		}
+        PublicIncludePaths.Add(IncDirectory);
+	}
+}
diff --git a/Engine/Plugins/Media/BinkMedia/Source/BinkMediaPlayerSDK/include/binkplugin.h b/Engine/Plugins/Media/BinkMedia/Source/BinkMediaPlayerSDK/include/binkplugin.h
new file mode 100644
--- /dev/null
+++ b/Engine/Plugins/Media/BinkMedia/Source/BinkMediaPlayerSDK/include/binkplugin.h
@@ -0,0 +1,354 @@
+// Copyright Epic Games, Inc. All Rights Reserved.
+// Licensed under the terms of a valid Unreal Engine license agreement,
+//   and the separate 'Unreal Engine End User License Agreement for Publishing'.
+
+#ifndef __BINKPLUGINH__
+#define __BINKPLUGINH__
+
+#include "egttypes.h"
+
+#ifdef BUILDING_FOR_UNREAL_ONLY
+
+// this is for static libs in unreal
+#ifdef __cplusplus
+#define PLUG_IN_FUNC_DEC( ret ) extern "C" ret 
+#define PLUG_IN_FUNC_DEF( ret ) extern "C" ret 
+#else
+#define PLUG_IN_FUNC_DEC( ret ) extern ret
+#define PLUG_IN_FUNC_DEF( ret ) ret
+#endif
+
+#endif
+
+#ifndef PLUG_IN_FUNC_DEC
+
+// normal DLLs if we aren't doing static libs
+#define PLUG_IN_FUNC_DEC( ret ) RADEXPFUNC ret RADEXPLINK
+#define PLUG_IN_FUNC_DEF( ret ) RADDEFFUNC ret RADLINK
+
+#endif
+
+#define BP_STRING_JOIN(arg1, arg2)              BP_STRING_JOIN_DELAY(arg1, arg2)
+#define BP_STRING_JOIN_DELAY(arg1, arg2)        BP_STRING_JOIN_IMMEDIATE(arg1, arg2)
+#define BP_STRING_JOIN_IMMEDIATE(arg1, arg2)    arg1 ## arg2
+
+
+enum BINKPLUGINAPI
+{
+#ifndef BUILDING_FOR_UNREAL_ONLY
+  BinkGL = 0,
+  BinkD3D9 = 1,
+  BinkD3D11 = 2,
+  BinkD3D12 = 3,
+  BinkMetal = 4,
+  BinkVulkan = 5,
+  BinkNDA = 6,     // any nda gpu api
+#else
+  BinkRHI = 0,     // Unreal RHI api
+#endif
+};
+
+// some platforms need allocators for cpu/gpu memory
+typedef void * BinkPluginAlloc_t( UINTa bytes, U32 alignment );
+typedef void   BinkPluginFree_t( void * ptr ); 
+typedef void * BinkPluginGPUAlloc_t( UINTa bytes, U32 alignment );
+typedef void   BinkPluginGPUFree_t( void * ptr, UINTa bytes ); 
+
+typedef struct BINKPLUGININITINFO
+{
+  // D3D12 = ID3D12Device*
+  // Vulkan = VkPhysicalDevice
+  // Metal = id< MTLDevice >
+  void * queue;
+
+  // Vulkan only = VkPhysicalDevice
+  void * physical_device; 
+  
+  // Vulkan only = formats for the sdr and hdr formats
+  U32 sdr_and_hdr_render_target_formats[2];
+
+  // general CPU allocators (can be zero on most platforms)
+  BinkPluginAlloc_t * alloc;
+  BinkPluginFree_t * free;
+
+  // gpu allocators (only used on some platforms
+  BinkPluginGPUAlloc_t * gpu_alloc;
+  BinkPluginGPUFree_t * gpu_free;
+} BINKPLUGININITINFO;
+
+
+// turn the plug in system on and off (these functions touch the graphics API)
+
+// device: 
+// D3D12 = ID3D12Device*
+// Vulkan = VkPhysicalDevice
+// Metal = id< MTLDevice >
+// Vulkan = logical device
+// NVNdevice*
+PLUG_IN_FUNC_DEC( S32 ) BinkPluginInit( void * device, BINKPLUGININITINFO * info, U32 graphics_api ); 
+PLUG_IN_FUNC_DEC( void ) BinkPluginShutdown( void );                     
+
+// spins through all binks, advancing frames and using gpu 
+//   this function will hit the gpu on D3D9 and GPUAssisted (not plain GL and D3D11)
+//   try to take less than ms_to_process (gotos will take longer, depending on goto time)
+PLUG_IN_FUNC_DEC( void ) BinkPluginProcessBinks( S32 ms_to_process ); 
+
+// spins through all binks, drawing everything, you can tell it what you want to draw
+//   if you want to draw the rendertargets at a different moment within the frame
+PLUG_IN_FUNC_DEC( void ) BinkPluginDraw( S32 draw_overlays, S32 draw_to_render_textures ); 
+
+// turn on and off IO for all Binks - if the buffer runs *completely* out, we still hit the disc
+PLUG_IN_FUNC_DEC( void ) BinkPluginIOPause( S32 IO_on );
+
+// Tell Bink to try to use GPU-assisted mode (once on, always on)
+PLUG_IN_FUNC_DEC( void ) BinkPluginTurnOnGPUAssist( void );
+
+// used to specify the sounds to open at playback
+enum BINKPLUGINSNDTRACK
+{
+  BinkSndNone               = 0, // don't open any sound tracks snd_track_start not used
+  BinkSndSimple             = 1, // based on filename, OR simply mono or stereo sound in track snd_track_start (default speaker spread)
+  BinkSndLanguageOverride   = 2, // mono or stereo sound in track 0, language track at snd_track_start
+  BinkSnd51                 = 3, // 6 mono tracks in tracks snd_track_start[0..5]
+  BinkSnd51LanguageOverride = 4, // 6 mono tracks in tracks 0..5, center language track at snd_track_start
+  BinkSnd71                 = 5, // 8 mono tracks in tracks snd_track_start[0..7]
+  BinkSnd71LanguageOverride = 6, // 8 mono tracks in tracks 0..7, center language track at snd_track_start
+};
+
+typedef struct BINKPLUGIN BINKPLUGIN;
+
+// used to specify the how the video should be buffered
+enum BINKPLUGINBUFFERING
+{
+  BinkStream                = 0, // stream the movie off the media during playback (caches about 1 second of video)
+  BinkPreloadAll            = 1, // loads the whole movie into memory at Open time (will block)
+  BinkStreamUntilResident   = 2, // streams the movie into a memory buffer as big as the movie, so it will be preloaded eventually)
+};
+
+
+// open and close a Bink file
+PLUG_IN_FUNC_DEC( BINKPLUGIN * ) BinkPluginOpen( char const * name, U32 snd_track_type, S32 snd_track_start, U32 buffering_type, U64 file_byte_offset );
+// Open a Bink file with a UTF-16 filename (Warning: Only works on Win32/64 platforms)
+PLUG_IN_FUNC_DEC( BINKPLUGIN * ) BinkPluginOpenUTF16( unsigned short const * name, U32 snd_track_type, S32 snd_track_start, U32 buffering_type, U64 file_byte_offset );
+// Close a Bink 
+PLUG_IN_FUNC_DEC( void ) BinkPluginClose( BINKPLUGIN * bnk );
+
+// function to preload a Bink completely into memory (so it can be played by multiple Binks)
+//   To use these function, just call them with a name and file offset, and if it's a Bink,
+//   the entire Bink file will be loaded into memory.  If you then call BinkPluginOpen with
+//   the same parameters, it will just read out of this memory, instead of re-loading it.
+//   64 Binks can be preloaded, and you shouldn't load too big of a file (these are
+//   are blocking calls).
+PLUG_IN_FUNC_DEC( S32 ) BinkPluginPreload( char const * name, U64 file_byte_offset );
+PLUG_IN_FUNC_DEC( void ) BinkPluginUnload( char const * name, U64 file_byte_offset );
+
+PLUG_IN_FUNC_DEC( S32 ) BinkPluginPreloadUTF16( unsigned short const * name, U64 file_byte_offset );
+PLUG_IN_FUNC_DEC( void ) BinkPluginUnloadUTF16( unsigned short const * name, U64 file_byte_offset );
+
+
+// set the path to load the library binaries from
+PLUG_IN_FUNC_DEC( void ) BinkPluginSetPath( char const * path );
+
+// error stuff
+PLUG_IN_FUNC_DEC( char const * ) BinkPluginError( void );
+PLUG_IN_FUNC_DEC( void ) BinkPluginSetError( char const * err );
+PLUG_IN_FUNC_DEC( void ) BinkPluginAddError( char const * err );
+
+typedef struct BINKPLUGININFO
+{
+  U64 BufferSize;
+  U64 BufferUsed;
+  U32 Width;
+  U32 Height;
+  U32 Frames;
+  U32 FrameNum;
+  U32 TotalFrames;
+  U32 FrameRate;
+  U32 FrameRateDiv;
+  U32 LoopsRemaining;
+  S32 ReadError;
+  S32 TexturesError;
+  U32 SndTrackType;          // BINKPLUGINSNDTRACK in use
+  U32 NumTracksRequested;    // Num tracks
+  U32 NumTracksOpened;
+  U32 SoundDropOuts;
+  U32 SkippedFrames;
+  U32 PlaybackState;         // 0 = playing, 1 = paused, 2 = gotoing, 3 = at end
+  F32 ProcessBinksFrameRate; // rate at which ProcessBinks is getting called (over last 32 processes)
+  F32 Alpha;
+} BINKPLUGININFO;
+
+// get playback info
+PLUG_IN_FUNC_DEC( void ) BinkPluginInfo( BINKPLUGIN * bnk, BINKPLUGININFO * info );
+
+// call one of these functions every frame
+PLUG_IN_FUNC_DEC( S32 ) BinkPluginScheduleToTexture( BINKPLUGIN * bnk, F32 x0, F32 y0, F32 x1, F32 y1, S32 depth, void * render_target_texture, U32 render_target_width, U32 render_target_height );
+PLUG_IN_FUNC_DEC( S32 ) BinkPluginScheduleOverlay( BINKPLUGIN * bnk, F32 x0, F32 y0, F32 x1, F32 y1, S32 depth );
+// call when everything is scheduled
+PLUG_IN_FUNC_DEC( void ) BinkPluginAllScheduled( void );
+
+// Pause playback
+PLUG_IN_FUNC_DEC( void ) BinkPluginPause( BINKPLUGIN * bnk, S32 pause_frame ); // 0 = resume, or framenumber to pause on, or -1 to pause immediately
+
+// start async jumping to frame (may draw same frame for long time until it finishes - can monitor with GetInfo)
+//   goto_frame is dest frame (fast if key frame). Use 0 to cancel a previous goto at the current position
+//   ms_per_process is how long Bink can spend decompressing in BinkPluginProcessBinks for this goto, use -1 for infinite wait
+PLUG_IN_FUNC_DEC( void ) BinkPluginGoto( BINKPLUGIN * bnk, S32 goto_frame, S32 ms_per_process ); 
+
+// set overall volume
+PLUG_IN_FUNC_DEC( void ) BinkPluginVolume( BINKPLUGIN * bnk, F32 vol ); // 0 to 1.0
+
+// set speaker volume
+//   BinkSndSimple = count must be 2 (l/r)
+//   BinkSndLanguageOverride = count must be 3 (l,r)/language
+//   BinkSnd51 = count must be 6 (front l/r),center,sub,(rear l/r),
+//   BinkSnd51LanguageOverride = 7 (front l/r),center,sub,(rear l/r),language
+//   BinkSnd71 = count must be 8 (front l/r),center,sub,(read l/r),(surr l/r)
+//   BinkSnd71LanguageOverride = 9 (front l/r),center,sub,(read l/r),(surr l/r), lang
+PLUG_IN_FUNC_DEC( void ) BinkPluginSpeakerVolumes( BINKPLUGIN * bnk, F32 * vols, U32 count ); // 0 to 1.0 
+
+// turn on/off video looping, loops == 0, infinite
+PLUG_IN_FUNC_DEC( void ) BinkPluginLoop( BINKPLUGIN * bnk, U32 loops );
+
+// sets HDR settings state used by BinkPluginScheduleToTexture and BinkPluginScheduleOverlay.
+//   tonemap = 0 (disabled), 1 (enabled), 2 (ST2084 PQ)
+//   exposure is a scaling factor that happens before tonemapping (1.0=normal, <1.0 darken, >1.0 brighten)
+//   output_luma = scales the tonemapped output to output this value as its maximum. 
+//     For HDR displays, set this to the max luma of the display. Typically 1000 nits to 2000 nits.
+PLUG_IN_FUNC_DEC( void ) BinkPluginSetHdrSettings( BINKPLUGIN * bnk, U32 tonemap, F32 exposure, S32 output_nits );
+
+// sets Alpha settings state used by BinkPluginScheduleToTexture and BinkPluginScheduleOverlay.
+//   alpha_value is just a constant blend value for entire video frame. 1 (default) opaque, 0 fully transparent.
+PLUG_IN_FUNC_DEC( void ) BinkPluginSetAlphaSettings( BINKPLUGIN * bnk, F32 alpha_value );
+
+// Only used with vulkan. Indexes into the formats specified on device init. See BINKPLUGINVULKANDEVICE 
+PLUG_IN_FUNC_DEC( void ) BinkPluginSetRenderTargetFormat( BINKPLUGIN * bnk, U32 format_idx );
+
+enum BINKPLUGINDRAWFLAGS {
+  // Decodes sRGB in the shader when drawing. Use this for example when rendering a 8-bit movie to a HDR texture format.
+  BinkDrawDecodeSRGB = 0x80000000,
+};
+
+// sets draw_flags state used by BinkPluginScheduleToTexture and BinkPluginScheduleOverlay.
+//   draw_flags is a bitmask of different options you can set when drawing. see BINKPLUGINDRAWFLAGS enum
+PLUG_IN_FUNC_DEC( void ) BinkPluginSetDrawFlags( BINKPLUGIN * bnk, S32 draw_flags );
+
+// limit speakers to certain number of speakers
+PLUG_IN_FUNC_DEC( void ) BinkPluginLimitSpeakers( U32 speaker_count );  // set to 1=mono max, 2=stereo max, 3=2.1 max, 4=4.0 max, 6=5.1 max, 8=7.1 max
+
+// for windows only, use directsound (otherwise uses Xaudio2) call before first open
+PLUG_IN_FUNC_DEC( S32 ) BinkPluginWindowsUseXAudioDevice( char const * strstr_device_name );  // name has to match some part of the windows device name "Rift" for example (or 0 for default)
+PLUG_IN_FUNC_DEC( S32 ) BinkPluginWindowsUseDirectSound( void );
+
+// for D3D9 windows only (call before and after device reset to reset GPU video textures)
+PLUG_IN_FUNC_DEC( void ) BinkPluginWindowsD3D9BeginReset( void );
+PLUG_IN_FUNC_DEC( void ) BinkPluginWindowsD3D9EndReset( void );
+
+PLUG_IN_FUNC_DEC( S32 ) BinkPluginGetPlatformInfo( U32 bink_platform_enum, void * output_ptr );
+
+/// Bink Image API ///
+
+typedef struct BINKIMAGE  BINKIMAGE;
+
+enum BINKIMAGEFORMAT {
+  BinkSurface32BGRA   =  5,
+  BinkSurface32RGBA   =  6,
+  BinkSurface565      = 10,
+  BinkSurface32ARGB   = 12,
+};
+
+// open and close a Bink file as an image
+PLUG_IN_FUNC_DEC( BINKIMAGE * ) BinkPluginOpenImage( char const * name, U32 *width, U32 *height, U64 file_byte_offset );
+// Open a Bink file as an image with a UTF-16 filename (Warning: Only works on Win32/64 platforms)
+PLUG_IN_FUNC_DEC( BINKIMAGE * ) BinkPluginOpenImageUTF16( unsigned short const * name, U32 *width, U32 *height, U64 file_byte_offset );
+// Read a Bink image into memory
+PLUG_IN_FUNC_DEC( void ) BinkPluginReadImage( BINKIMAGE * bnki, void *dest, S32 destpitch, U32 format_flags );
+// Close a Bink image
+PLUG_IN_FUNC_DEC( void ) BinkPluginCloseImage( BINKIMAGE * bnki );
+
+typedef struct BINKPLUGINFRAMEINFO
+{
+  // (Metal), id<MTLCommandBuffer>
+  // (GNM), Gnmx::LightweightGfxContext *
+  // (NVN), NVNcommandBuffer *
+  // (AGC), DrawCommandBuffer *
+  void * cmdBuf; 
+
+  // (D3D12), swap buffer ptr. On D3D12, must have the flags D3D12_RESOURCE_FLAG_ALLOW_RENDER_TARGET
+  // (Metal), id<MTLTexture>
+  // (Vulkan), VkImage
+  void * screen_resource; 
+  
+  // (D3D12), D3D12_RESOURCE_STATES, what the screen resource starts in and is restored back to.
+  // (Vulkan), VkImageLayout, what layout the screen starts in and is restored back to.
+  U32 screen_resource_state; 
+  
+  // width and height of the screen resource
+  U32 width, height;
+  
+  // (Vulkan), screen is sdr or hdr (formats set in INITINFO above)
+  U32 sdr_or_hdr; // on vulkan, 0=sdr screen, 1=hdr screen
+} BINKPLUGINFRAMEINFO;
+
+PLUG_IN_FUNC_DEC( void ) BinkPluginSetPerFrameInfo( BINKPLUGINFRAMEINFO * info );  
+
+// dynamic linking typedefs - you can use DO_BINKPLUGIN_PROCS() to create them and such
+#define DO_BINKPLUGIN_BINKONLY_PROCS() \
+  ProcessProc(  8, S32,            BinkPluginInit,                  void *device, BINKPLUGININITINFO * info, U32 graphics_api ) \
+  ProcessProc(  0, void,           BinkPluginShutdown,              void ) \
+  ProcessProc( 24, BINKPLUGIN *,   BinkPluginOpen,                  char const * name, U32 snd_track_type, S32 snd_track_start, U32 buffering_type, U64 file_byte_offset ) \
+  ProcessProc( 24, BINKPLUGIN *,   BinkPluginOpenUTF16,             unsigned short const * name, U32 snd_track_type, S32 snd_track_start, U32 buffering_type, U64 file_byte_offset ) \
+  ProcessProc(  4, void,           BinkPluginClose,                 BINKPLUGIN * bnk ) \
+  ProcessProc(  8, void,           BinkPluginInfo,                  BINKPLUGIN * bnk, BINKPLUGININFO *info ) \
+  ProcessProc( 12, void,           BinkPluginGoto,                  BINKPLUGIN * bnk, S32 goto_frame, S32 ms_per_process ) \
+  ProcessProc(  4, void,           BinkPluginProcessBinks,          S32 ms_to_process ) \
+  ProcessProc(  4, void,           BinkPluginLimitSpeakers,         U32 num_speakers ) \
+  ProcessProc(  8, void,           BinkPluginLoop,                  BINKPLUGIN * bnk, U32 loops ) \
+  ProcessProc(  8, void,           BinkPluginPause,                 BINKPLUGIN * bnk, S32 pause_frame ) \
+  ProcessProc(  8, void,           BinkPluginDraw,                  S32 draw_overlays, S32 draw_to_render_textures ) \
+  ProcessProc(  4, void,           BinkPluginSetPerFrameInfo,       BINKPLUGINFRAMEINFO * info ) \
+  ProcessProc( 36, S32,            BinkPluginScheduleToTexture,     BINKPLUGIN *bnk, F32 x0, F32 y0, F32 x1, F32 y1, S32 depth, void * render_target_texture, U32 render_target_width, U32 render_target_height ) \
+  ProcessProc( 24, S32,            BinkPluginScheduleOverlay,       BINKPLUGIN *bnk, F32 x0, F32 y0, F32 x1, F32 y1, S32 depth ) \
+  ProcessProc(  0, void,           BinkPluginAllScheduled,          void ) \
+  ProcessProc(  4, void,           BinkPluginIOPause,               S32 IO_on ) \
+  ProcessProc(  0, void,           BinkPluginTurnOnGPUAssist,       void ) \
+  ProcessProc(  0, char const *,   BinkPluginError,                 void ) \
+  ProcessProc(  8, void,           BinkPluginVolume,                BINKPLUGIN * bnk, F32 vol ) \
+  ProcessProc( 12, void,           BinkPluginSpeakerVolumes,        BINKPLUGIN * bnk, F32 * vols, U32 count ) \
+  ProcessProc( 16, void,           BinkPluginSetHdrSettings,        BINKPLUGIN * bnk, U32 tonemap, F32 exposure, S32 output_nits ) \
+  ProcessProc(  8, void,           BinkPluginSetAlphaSettings,      BINKPLUGIN * bnk, F32 alpha_value ) \
+  ProcessProc(  8, void,           BinkPluginSetRenderTargetFormat, BINKPLUGIN * bnk, U32 format ) \
+  ProcessProc(  4, void,           BinkPluginSetPath,               char const * path ) \
+  ProcessProc(  8, S32,            BinkPluginGetPlatformInfo,       U32 bink_platform_enum, void * output_ptr ) \
+  ProcessProc(  8, void,           BinkPluginSetDrawFlags,          BINKPLUGIN * bnk, S32 draw_flags ) \
+  ProcessProc(  8, S32,            BinkPluginPreload,               char const * name, U64 file_byte_offset ) \
+  ProcessProc(  8, void,           BinkPluginUnload,                char const * name, U64 file_byte_offset ) \
+  ProcessProc(  8, S32,            BinkPluginPreloadUTF16,          unsigned short const * name, U64 file_byte_offset ) \
+  ProcessProc(  8, void,           BinkPluginUnloadUTF16,           unsigned short const * name, U64 file_byte_offset ) \
+
+#ifdef BUILDING_FOR_UNREAL_ONLY
+
+#define DO_BINKPLUGIN_PROCS() \
+  DO_BINKPLUGIN_BINKONLY_PROCS() \
+
+#else
+
+#define DO_BINKPLUGIN_PROCS() \
+  DO_BINKPLUGIN_BINKONLY_PROCS() \
+  ProcessProc( 16, BINKIMAGE *,    BinkPluginOpenImage,             char const * name, U32 *width, U32 *height, U64 file_byte_offset ) \
+  ProcessProc( 16, BINKIMAGE *,    BinkPluginOpenImageUTF16,        unsigned short const * name, U32 *width, U32 *height, U64 file_byte_offset ) \
+  ProcessProc( 16, void,           BinkPluginReadImage,             BINKIMAGE * bnki, void *dest, S32 destpitch, U32 format_flags ) \
+  ProcessProc(  4, void,           BinkPluginCloseImage,            BINKIMAGE * bnki ) \
+
+#endif
+
+#define DO_BINKPLUGIN_WIN_PROCS() \
+  ProcessProc(  4, S32,            BinkPluginWindowsUseXAudioDevice, char const * strstr_device_name ) \
+  ProcessProc(  0, S32,            BinkPluginWindowsUseDirectSound,  void ) \
+
+#define ProcessProc( bytes, ret, name, ...) typedef ret BP_STRING_JOIN( name, Proc )(__VA_ARGS__); extern BP_STRING_JOIN( name, Proc ) * BP_STRING_JOIN( p, name );
+  DO_BINKPLUGIN_PROCS()
+  DO_BINKPLUGIN_WIN_PROCS()
+#undef ProcessProc
+
+#endif
diff --git a/Engine/Plugins/Media/BinkMedia/Source/BinkMediaPlayerSDK/include/binktextures.h b/Engine/Plugins/Media/BinkMedia/Source/BinkMediaPlayerSDK/include/binktextures.h
new file mode 100644
--- /dev/null
+++ b/Engine/Plugins/Media/BinkMedia/Source/BinkMediaPlayerSDK/include/binktextures.h
@@ -0,0 +1,477 @@
+// Copyright Epic Games, Inc. All Rights Reserved.
+// Licensed under the terms of a valid Unreal Engine license agreement,
+//   and the separate 'Unreal Engine End User License Agreement for Publishing'.
+
+/*
+This is Epics's high level API for using 3D hardware to do Bink GPU-assisted
+decoding and color conversion. It is supported on every 3D platform that 
+Bink runs on.
+
+It's a simple API, so you should see your platform's example code
+to see all the nitty-gritty details.
+
+
+There two functions to start and stop Bink textures:
+
+  Create_Bink_shaders:  This function creates the pixel shaders we use to
+    do the color space conversion. On most platforms, you need to provide
+    a pointer to GPU device pointer. 
+    All of the other functions hang off the object that is returned here.
+
+  Free_Bink_shaders:  Frees the pixel shaders created in Create_Bink_shaders.
+
+There are four textures functions:
+
+  Create_Bink_textures:  This function takes BINK object, creates the texture
+    resources to render it quickly and binds them together for rendering.
+  
+  Free_Bink_textures:  Frees the resources allocated in Create_Bink_textures.
+
+  Draw_Bink_textures:  Renders the textures onto the screen.
+
+  Start_Bink_texture_update:  This function starts the texture updating process.
+    It should be called before you call BinkDoFrame (synchronous or async).
+    On some platforms, this locks/maps the texture for updating; on consoles, it
+    makes sure that the GPU is done reading from the texture buffers that we're
+    about to overwrite.
+
+  Finish_Bink_texture_update:  This function finishes the texture updating process.
+    It should be called after BinkDoFrame processing is done (in async mode: after
+    BinkDoFrameAsyncWait returns true). This will unlock/unmap textures (where it
+    applies) or issue a blit to transfer the new image data. On most consoles with
+    unified memory, we decode directly into GPU-readable memory, and this operation
+    is a no-op.
+
+On Windows D3D9, there are two other functions:
+
+  Before_Reset_Bink_textures: Call before you call the D3D Reset function.
+
+  After_Reset_Bink_textures: Call after you call the D3D Reset function.
+
+For Windows GL or D3D11, you simply free and re-create the textures completely
+  whenever the device is lost or reset.
+
+
+
+So, basically, playback works like this:
+
+  1) Call Create_Bink_shaders to create the pixel shaders.
+  
+  2) Open the Bink file with the BINKNOFRAMEBUFFERS flag (the Binktexture
+     implementations will allocate the memory to decompress into).
+
+  3) Call Create_Bink_textures to create the textures and bind the textures
+     into the Bink.
+
+  4) Call Start_Bink_texture_update before BinkDoFrame.
+
+  5) Call BinkDoFrame to decompress a video frame.
+
+  6) Call Finish_Bink_texture_update after BinkDoFrame (or after
+     BinkDoFrameAsyncWait returns true, if you're using async decoding).
+
+  7) Draw the frame using Draw_Bink_textures.
+
+
+And that's it! (Skipping over a few details - see the examples for all 
+the details...)
+
+Should drop in really quickly and it hides a ton of platform specific ugliness!
+*/
+
+#ifndef BINKTEXTURESH
+#define BINKTEXTURESH
+
+#ifndef __BINKH__
+#include "bink.h"
+#endif
+
+RADDEFSTART
+
+//=============================================================================
+
+// On some platforms, you might have one app that uses BinkTextures with multiple
+//   graphics APIs - these defines let you have multiple providers in one
+//   EXE by suffixing the names of D3D9 and D3D11 (GL gets the non-suffix).
+// This does mean that you must include D3D9 and/or D3D11 headers before
+//   including BinkTextures.h!
+
+#if defined( BINKTEXTURESGPUAPITYPE )
+
+  #if defined( __d3d12_x_h__ ) || defined( __d3d12_xs_h__ ) ||defined( __d3d12_h__ ) || defined( BINKD3D12FUNCTIONS ) || defined( BINKD3D12GPUFUNCTIONS )
+    #ifdef BINKD3D12GPUFUNCTIONS
+      #define BINKTEXTURESSUFFIX D3D12GPU
+    #else
+      #define BINKTEXTURESSUFFIX D3D12
+    #endif
+  #elif defined( __d3d11_x_h__ ) || defined( __d3d11_h__ ) || defined( BINKD3D11FUNCTIONS ) || defined( BINKD3D11GPUFUNCTIONS )
+    #ifdef BINKD3D11GPUFUNCTIONS
+      #define BINKTEXTURESSUFFIX D3D11GPU
+    #else
+      #define BINKTEXTURESSUFFIX D3D11
+    #endif
+  #elif defined( _D3D9_H_ ) || defined( BINKD3D9FUNCTIONS )
+    #define BINKTEXTURESSUFFIX D3D9
+  #elif defined( BINKMETALFUNCTIONS )
+    #define BINKTEXTURESSUFFIX Metal
+  #elif defined( BINKVULKANFUNCTIONS )
+    #define BINKTEXTURESSUFFIX Vulkan
+  #elif defined(__gl_es20_h_) || defined(__gl_h_) || defined(__gl2_h_) || defined(__gl3_h_) || defined( BINKGLFUNCTIONS ) || defined(BINKGLGPUFUNCTIONS)
+    #ifdef BINKGLGPUFUNCTIONS
+      #define BINKTEXTURESSUFFIX GLGPU
+    #else
+      #define BINKTEXTURESSUFFIX GL
+    #endif
+  #elif defined(BINKNDAFUNCTIONS)
+    #define BINKTEXTURESSUFFIX BINKNDAFUNCTIONS
+  #elif defined(BINKRHIFUNCTIONS)
+    #define BINKTEXTURESSUFFIX RHI
+  #else
+    #ifndef BINKTEXTURESSUFFIX
+      #define BINKTEXTURESSUFFIX Null
+    #endif
+  #endif
+
+  #define Create_Bink_shaders        RR_STRING_JOIN( Create_Bink_shaders,       BINKTEXTURESSUFFIX )
+
+#endif
+
+typedef struct BINKSHADERS BINKSHADERS;   // defined below
+typedef struct BINKTEXTURES BINKTEXTURES; // defined below
+
+//=============================================================================
+
+// For some APIs, you need to pass structs as argument to our functions.
+// This is where they get defined.
+
+#if ( defined( __d3d12_h__ ) || defined( __d3d12_x_h__ ) || defined( __d3d12_xs_h__ ) ) && !defined( BINKTEXTURESD3D12STRUCTSDEF )
+
+// BinkGPU can call you back at the beginning/end of its compute command lists in case you want to
+// insert timestamp queries (or other markup) into the command list.
+typedef void BINKGPUCOMMANDLISTCALLBACKD3D12( BINKTEXTURES * bink_tex, ID3D12GraphicsCommandList * cl, BOOL is_end, void * user_ptr );
+
+// This is what you pass as "device" to Create_shaders.
+typedef struct BINKCREATESHADERSD3D12 {
+  // The D3D12 device to use
+  ID3D12Device * device;
+  
+  // Command queue to submit work to in BinkGPU mode. Ignored in CPU decoding mode.
+  ID3D12CommandQueue * gpu_mode_command_queue;
+
+  // The fence to use for synchronization. Needs to be updated from app code!
+  ID3D12Fence * fence; 
+
+  // Prototype graphics pipeline state. Create Bink pipelines for the
+  // specified set of rasterizer state (multisampling or not),
+  // render target formats, etc. See binktexturesD3D12.cpp for
+  // details.
+  D3D12_GRAPHICS_PIPELINE_STATE_DESC prototype;
+
+  // Callback to run at beginning/end of BinkGPU command list. Ignored in CPU decoding mode.
+  BINKGPUCOMMANDLISTCALLBACKD3D12 * gpu_command_list_callback;
+  void * gpu_command_list_user_ptr;
+} BINKCREATESHADERSD3D12;
+
+// This is what you pass as context to Draw_textures.
+typedef struct BINKAPICONTEXTD3D12 {
+  // ---- Input parameters you provide
+
+  // Command list to render to
+  ID3D12GraphicsCommandList * command_list; // Command list to render to
+
+  // Fence value you're going to write via the command queue after executing
+  // the command_list. The sequence of fence values written needs to be
+  // monotonically increasing over time. This is what we use to synchronize
+  // buffer upload ands resource management!
+  UINT64 fence_value;
+
+  // ---- Output parameters we return
+
+  // In GPU decode mode, the draw command list may only run after decoding
+  // the frame has finished. Thus you need to call "Wait" on your graphics
+  // command queue to make sure our decode fence is signaled to the correct
+  // value before decoding can start.
+  //
+  // NOTE: "out_gpu_decode_fence" is not AddRef'd before being returned, so please
+  // don't call Release on it!
+
+  ID3D12Fence * out_gpu_decode_fence;   // In GPU mode, this is the fence you have to call wait on before submitting command_list.
+  UINT64 out_gpu_decode_fence_val;      // In GPU mode, this is the fence value you have to wait for before submitting command_list.
+} BINKAPICONTEXTD3D12;
+
+// We promise not to request more SRV descriptors than this per Bink, for any D3D12 BinkTextures implementation.
+// So if you want to use a simple free list allocator for Bink video descriptors, you can just use this as your
+// allocation granularity.
+#define BINKTEXTURESD3D12_MAX_SRVS_PER_BINK     128
+
+#define BINKTEXTURESD3D12STRUCTSDEF
+
+#endif
+
+#ifndef BINKCREATESHADERSPIPEANDQUEUEd
+#define BINKCREATESHADERSPIPEANDQUEUEd
+
+// Create_Bink_Shaders for some nda GPU implementations allows you to select which async compute pipe and queue ID to run
+// decoding jobs on. If not specified, we use our default. (At the time of writing, this is pipe 6, queue 1).
+// Fill out this struct and pass a pointer to it as "device" to Create_Bink_shaders.
+typedef struct BINKCREATESHADERSPIPEANDQUEUE 
+{
+  // Which async compute pipe to use for BinkGPU. Ignored for CPU decoding.
+  // Corresponds to the "pipeId" in ComputeQueue::initialize.
+  U32 pipe_id;
+
+  // Which queue to use for BinkGPU on the given compute pipe. Ignored for CPU decoding.
+  // Corresponds to the "queueId" in ComputeQueue::initialize.
+  U32 queue_id;
+} BINKCREATESHADERSPIPEANDQUEUE;
+
+#endif
+
+#ifndef BINKTEXTURES_METAL_DEF
+#define BINKTEXTURES_METAL_DEF
+
+// Create_Bink_Shaders for Metal needs a few more parameters - pass a pointer to this structure.
+
+typedef struct BINKCREATESHADERSMETAL 
+{
+  void * mtlDevice;
+  void * stateDesc[2];          // MTLRenderPipelineDescriptor *
+  void * mtlQueue;              // used on GPU-assist Bink Metal only
+  U32 render_target_formats[2]; // MTLPixelFormat(s) that we can render to
+} BINKCREATESHADERSMETAL;
+
+
+typedef struct BINKAPICONTEXTMETAL
+{
+  void * command_buffer;  // id <MTLCommandBuffer> 
+  void * command_encoder; // id <MTLRenderCommandEncoder>
+  void * frame_texture;   // View.currentDrawable.texture
+  U32 format_idx;         // Index into render_target_formats. See BINKCREATESHADERSMETAL
+} BINKAPICONTEXTMETAL;
+
+#endif
+
+#ifndef BINKTEXTURES_VULKAN_DEF
+#define BINKTEXTURES_VULKAN_DEF
+typedef struct BINKCREATESHADERSVULKAN 
+{
+  void *physical_device;        // VkPhysicalDevice
+  void *logical_device;         // VkDevice
+  void *queue;                  // VkQueue (needs graphics+compute capabilities where compute is only used for gpu decode)
+  U32 render_target_formats[2]; // VkFormat(s) that we can render to
+} BINKCREATESHADERSVULKAN;
+
+typedef struct BINKAPICONTEXTVULKAN
+{
+  void * cmdbuf;  // VkCommandBuffer
+  U32 format_idx; // Index into render_target_formats. See BINKCREATESHADERSVULKAN
+} BINKAPICONTEXTVULKAN;
+#endif
+
+//=============================================================================
+
+// Creates shaders that Bink uses.
+//   The value that you pass for device depends on the graphics API:
+//     D3D9 -    The d3d9 device (type LPDIRECT3DDEVICE9)
+//     D3D11 -   The d3d11 device (type D3D11Device *)
+//     D3D12 -   Pointer to a filled-out BINKCREATESHADERSD3D12 struct
+//     GL -      Parameter not used, just pass 0
+//     PS4 -     Pass 0 for default parameters, or a pointer to a filled-out BINKCREATESHADERSPS4 struct
+//     Metal   - Pointer to a filled out BINKCREATESHADERSMETAL struct
+BINKSHADERS * Create_Bink_shaders( void * device );
+
+// free our shaders
+typedef void Free_Bink_shaders_ft( BINKSHADERS * shaders );
+
+//=============================================================================
+
+// allocate the textures that we'll need - shaders defines the device we use to
+//   create them. The user_ptr gets stored in the BINKTEXTURES and also passed
+//   to any underlying "alloc" functions we call.
+typedef BINKTEXTURES * Create_Bink_textures_ft( BINKSHADERS * shaders, HBINK bink, void *user_ptr );
+
+//=============================================================================
+
+// frees the textures
+typedef void Free_Bink_textures_ft( BINKTEXTURES * textures );
+
+//=============================================================================
+// Bink texture updating
+
+// Start updating textures for this frame. Call before DoFrame.
+typedef void Start_Bink_texture_update_ft( BINKTEXTURES * textures );
+
+// Finish updating textures for this frame. Call after DoFrame is done.
+typedef void Finish_Bink_texture_update_ft( BINKTEXTURES * textures );
+
+
+//=============================================================================
+
+// Bink texture drawing functions
+
+// Draws the textures as a screen-aligned quad.
+
+//   If you pass a zero for "shaders", we'll use the shader used to create the
+//       textures (which is normally what you want).
+//   The value that you pass for "graphics_context" depends on the graphics API:
+//     D3D9 -    Parameter not used, just pass 0
+//     D3D11 -   The d3d11 device context (type ID3D11DeviceContext *)
+//     D3D12 -   Pointer to a BINKAPICONTEXTD3D12 struct with all fields filled out
+//     GL -      Parameter not used, just pass 0
+//     PS3 -     The Gcm context (usually, gCellGcmCurrentContext, type CellGcmContextData *)
+//     PS4 -     The Gfx context pointer to use (type sce::Gnmx::GfxContext *)
+//     Xbox360 - Parameter not used, just pass 0
+//     XboxOne - The Xbox device context (type ID3D11DeviceContextX *)
+//     NX -      The NX context pointer to use (type nn::gfx::CommandBuffer *)
+//     Metal -   Pointer to a BINKAPICONTEXTMETAL struct with all the fields filled out
+//     PS5 -     Pointer to the DrawCommandBuffer to use (no StateBuffer or Binder required)
+typedef void Draw_Bink_textures_ft( BINKTEXTURES * textures,
+                                    BINKSHADERS * shaders,
+                                    void * graphics_context );
+
+// Set the position to draw the video within the viewport.
+//   If you video doesn't move, you can call this once,
+//     otherwise call before each call to Draw_Bink_texture.
+//   By default, the video fills the video port.
+//
+//   0,0, 1,1 is the default which completely fills the viewport
+typedef void Set_Bink_draw_position_ft( BINKTEXTURES * textures,
+                                        F32 x0, F32 y0, F32 x1, F32 y1 );
+
+
+//   0,0, 1,0, 0,1 is the default which completely fills the viewport
+typedef void Set_Bink_draw_corners_ft( BINKTEXTURES * textures,
+                                        F32 Ax, F32 Ay, F32 Bx, F32 By, F32 Cx, F32 Cy );
+
+
+// Set draw source rect (from 0 to 1, in texturespace, so 0,0, 1,1 is whole texture)
+typedef void Set_Bink_source_rect_ft( BINKTEXTURES * textures,
+                                      F32 u0, F32 v0, F32 u1, F32 v1 );
+
+// Or these flags with draw_type_and_flags
+#define BINKDRAWFLAG_SRGB_DECODE (0x80000000)
+
+// Set the alpha settings for drawing.
+//   alpha_value is just a constant blend value for entire video frame
+//   draw_type is 0 = normal alpha, 1 = pre-multiplied alpha, 2 = opaque & copy alpha
+//   draw_type also contains flags in the high bits:
+//     BINKDRAWFLAG_SRGB_DECODE - decode sRGB to linear when rendering the video.
+typedef void Set_Bink_alpha_settings_ft( BINKTEXTURES * textures,
+                                         F32 alpha_value,
+                                         S32 draw_type_and_flags );
+
+// Set the hdr settings for drawing.
+//   tonemap specifies whether you want linear output (0), or filmic tonemapped output (1), or ST2084 HDR output (2).
+//   exposure is a scaling factor that happens before tonemapping (1.0=normal, <1.0 darken, >1.0 brighten)
+//   out_luma specifies the maximum luminance to output when tonemapping
+typedef void Set_Bink_hdr_settings_ft( BINKTEXTURES * texutres,
+                                       S32 tonemap,
+                                       F32 exposure, 
+                                       S32 output_nits );
+
+
+//=============================================================================
+
+// For D3D9 only, these functions are called before and after you reset the D3D9 device
+typedef void Before_Reset_Bink_textures_ft( BINKTEXTURES * textures );
+typedef S32  After_Reset_Bink_textures_ft( BINKTEXTURES * textures );
+
+//=============================================================================
+
+#ifndef BINKTEXTURESTRUCTS
+#define BINKTEXTURESTRUCTS
+
+#define Create_Bink_textures( shader, bink, user_ptr )   (shader)->Create_textures( shader, bink, user_ptr )
+#define Free_Bink_shaders( shader )                      (shader)->Free_shaders( shader )
+
+#define Start_Bink_texture_update( textures )            (textures)->Start_texture_update( textures )
+#define Finish_Bink_texture_update( textures )           (textures)->Finish_texture_update( textures )
+#define Draw_Bink_textures( textures, shaders, gr_ctx )  (textures)->Draw_textures( textures, shaders, gr_ctx )
+#define Set_Bink_draw_position( textures, x0,y0, x1,y1 ) (textures)->Set_draw_position( textures, x0,y0, x1,y1 )
+#define Set_Bink_draw_corners( textures, Ax,Ay, Bx,By, Cx,Cy ) (textures)->Set_draw_corners( textures, Ax,Ay, Bx,By, Cx,Cy )
+#define Set_Bink_source_rect( textures, u0,v0, u1,v1 )   (textures)->Set_source_rect( textures, u0,v0, u1,v1 )
+#define Set_Bink_alpha_settings( textures, alpha, draw_type_and_flags ) (textures)->Set_alpha_settings( textures, alpha, draw_type_and_flags )
+#define Set_Bink_hdr_settings( textures, tonemap, exposure, out_nits ) (textures)->Set_hdr_settings( textures, tonemap, exposure, out_nits )
+#define Before_Reset_Bink_textures( textures )           {if ((textures)->Before_Reset_textures) (textures)->Before_Reset_textures( textures );}
+#define After_Reset_Bink_textures( textures )            {if ((textures)->After_Reset_textures) (textures)->After_Reset_textures( textures );}
+#define Free_Bink_textures( textures )                   (textures)->Free_textures( textures )
+
+struct BINKSHADERS
+{
+  Create_Bink_textures_ft * Create_textures;
+  Free_Bink_shaders_ft * Free_shaders;
+
+  // you can use this for whatever you want
+  UINTa user_data[ 4 ];
+};
+
+struct BINKTEXTURES
+{
+  // wrap around doframe
+  Start_Bink_texture_update_ft * Start_texture_update;
+  Finish_Bink_texture_update_ft * Finish_texture_update;
+
+  // draw related functions
+  Draw_Bink_textures_ft * Draw_textures;
+  Set_Bink_draw_position_ft * Set_draw_position;
+  Set_Bink_draw_corners_ft * Set_draw_corners;
+  Set_Bink_source_rect_ft * Set_source_rect;
+  Set_Bink_alpha_settings_ft * Set_alpha_settings;
+  Set_Bink_hdr_settings_ft * Set_hdr_settings;
+
+  // D3D9 only (everywhere else is null)
+  Before_Reset_Bink_textures_ft * Before_Reset_textures;
+  After_Reset_Bink_textures_ft * After_Reset_textures;
+
+  Free_Bink_textures_ft * Free_textures;
+
+  // user pointer specified at Create_Bink_textures time
+  void * user_ptr;
+
+  // these are the platform specific texture pointers that you can draw manually with
+  void * Ytexture;
+  void * cRtexture;
+  void * cBtexture;
+  void * Atexture;
+  void * Htexture;
+
+  // you can use this for whatever you want
+  UINTa user_data[ 4 ];
+};
+
+#endif
+//=============================================================================
+
+
+// optional clean up, for when including the header multiple times for different APIs
+#if defined(BINKTEXTURESCLEANUP)
+  #undef Create_Bink_shaders 
+
+  #ifdef BINKD3D9FUNCTIONS
+    #undef BINKD3D9FUNCTIONS
+  #endif
+  #ifdef BINKD3D11FUNCTIONS
+    #undef BINKD3D11FUNCTIONS
+  #endif
+  #ifdef BINKD3D11GPUFUNCTIONS
+    #undef BINKD3D11GPUFUNCTIONS
+  #endif
+  #ifdef BINKD3D12FUNCTIONS
+    #undef BINKD3D12FUNCTIONS
+  #endif
+  #ifdef BINKD3D12GPUFUNCTIONS
+    #undef BINKD3D12GPUFUNCTIONS
+  #endif
+  #ifdef BINKNDAFUNCTIONS
+    #undef BINKNDAFUNCTIONS 
+  #endif
+  #ifdef BINKTEXTURESSUFFIX
+    #undef BINKTEXTURESSUFFIX 
+  #endif
+  #undef BINKTEXTURESH
+  #undef BINKTEXTURESCLEANUP
+#endif
+
+RADDEFEND
+
+#endif
diff --git a/Engine/Plugins/Media/BinkMedia/Source/BinkMediaPlayerSDK/include/binktiny.h b/Engine/Plugins/Media/BinkMedia/Source/BinkMediaPlayerSDK/include/binktiny.h
new file mode 100644
--- /dev/null
+++ b/Engine/Plugins/Media/BinkMedia/Source/BinkMediaPlayerSDK/include/binktiny.h
@@ -0,0 +1,89 @@
+#ifndef __BINKH__
+#define __BINKH__
+
+typedef struct BINK * HBINK;
+
+#define BINKALPHA              0x00100000L // Decompress alpha plane (if present)
+#define BINKHDR                0x00000004L // Video is an HDR video
+#define BINKUSETRIPLEBUFFERING 0x00000008L // Use triple buffering in the framebuffers
+
+typedef struct BINKRECT {
+  S32 Left,Top,Width,Height;
+} BINKRECT;
+
+#define BINKMAXDIRTYRECTS 1
+
+typedef struct BINKPLANE
+{
+  void * Buffer;
+  S32 Allocate;
+  U32 BufferPitch;
+} BINKPLANE;
+
+typedef struct BINKFRAMEPLANESET
+{
+  BINKPLANE YPlane;
+  BINKPLANE cRPlane;
+  BINKPLANE cBPlane;
+  BINKPLANE APlane;
+  BINKPLANE HPlane;
+} BINKFRAMEPLANESET;
+
+#define BINKMAXFRAMEBUFFERS 3
+
+typedef struct BINKFRAMEBUFFERS
+{
+  S32 TotalFrames;
+  U32 YABufferWidth;
+  U32 YABufferHeight;
+  U32 cRcBBufferWidth;
+  U32 cRcBBufferHeight;
+
+  U32 FrameNum;  // 0 to (TotalFrames-1)
+  BINKFRAMEPLANESET Frames[ BINKMAXFRAMEBUFFERS ];
+} BINKFRAMEBUFFERS;
+
+
+typedef struct BINK {
+  U32 Width;                  // Width (1 based, 640 for example)
+  U32 Height;                 // Height (1 based, 480 for example)
+  U32 Frames;                 // Number of frames (1 based, 100 = 100 frames)
+  U32 FrameNum;               // Frame to *be* displayed (1 based)
+
+  U32 LastFrameNum;           // Last frame decompressed or skipped (1 based)
+  U32 FrameRate;              // Frame Rate Numerator
+  U32 FrameRateDiv;           // Frame Rate Divisor (frame rate=numerator/divisor)
+  U32 ReadError;              // Non-zero if a read error has ocurred
+
+  U32 OpenFlags;              // flags used on open
+  U32 BinkType;               // Bink flags
+  U32 LargestFrameSize;       // Largest frame size
+  U32 FrameSize;              // The current frame's size in bytes
+
+  U32 SndSize;                // The current frame sound tracks' size in bytes
+  U32 FrameChangePercent;     // very rough percentage of the frame that changed
+  S32 NumTracks;              // how many tracks
+  S32 NumRects;
+
+  U32 soundon;                // sound turned on?
+  U32 videoon;                // video turned on?
+  U32 needio;                 // set to 1, if we need an io before the next binkframe
+  S32 closing;                // are we closing?
+
+  BINKRECT FrameRects[BINKMAXDIRTYRECTS];// Dirty rects from BinkGetRects
+
+  F32 ColorSpace[16];         // 4x4 matrix to use for colorspace conversion
+                              // in HDR, 0 = smpte2084(maxnits/10k)*0.5, 1/2=ct/cp scale, 3/4=ct/cp bias, 5 = maxnits
+
+  U64 FileOffset;             // Offset into the file where the Bink starts        
+  U64 Highest1SecRate;        // Highest 1 sec data rate
+
+  BINKFRAMEBUFFERS * FrameBuffers; // Bink frame buffers that we decompress to
+
+} BINK;
+
+RADEXPFUNC void RADEXPLINK BinkGetFrameBuffersInfo( HBINK bink, BINKFRAMEBUFFERS * fbset );
+RADEXPFUNC S32  RADEXPLINK BinkAllocateFrameBuffers( HBINK bp, BINKFRAMEBUFFERS * set, U32 minimum_alignment );
+RADEXPFUNC void RADEXPLINK BinkRegisterFrameBuffers( HBINK bink, BINKFRAMEBUFFERS * fbset );
+
+#endif
diff --git a/Engine/Plugins/Media/BinkMedia/Source/BinkMediaPlayerSDK/include/egttypes.h b/Engine/Plugins/Media/BinkMedia/Source/BinkMediaPlayerSDK/include/egttypes.h
new file mode 100644
--- /dev/null
+++ b/Engine/Plugins/Media/BinkMedia/Source/BinkMediaPlayerSDK/include/egttypes.h
@@ -0,0 +1,625 @@
+// Copyright Epic Games, Inc. All Rights Reserved.
+// Licensed under the terms of a valid Unreal Engine license agreement,
+//   and the separate 'Unreal Engine End User License Agreement for Publishing'.
+
+//===============================================
+// (C) Copyright 1994-2022 Epic Games Tools, LLC  
+//===============================================
+
+#ifndef __EGTTYPESH__
+#define __EGTTYPESH__
+
+#if !defined(__RADTYPESH__) && !defined(__RADRR_COREH__) 
+#define __RADTYPESH__
+#define __RADRR_COREH__ // block old rr_core
+
+#define RADCOPYRIGHT "Copyright (C) 1994-2022, Epic Games Tools, LLC"
+
+#if !defined(__RADRES__) // don't include anything for resource compiles
+
+//  __RAD32__ means at least 32 bit code (always defined)
+//  __RAD64__ means 64 bit addressing (64-bit OSes only) (__RAD64REGS__ will be set)
+//  __RAD64REGS__ means 64 bit registers
+
+//  __RADNT__ means Win32 and Win64 desktop
+//  __RADWINRT__ means Windows Store/Phone App (x86, x64, arm)
+//  __RADWIN__ means win32, win64, windows store/phone, consoles
+//  __RADWINRTAPI__ means Windows RT API (Win Store, Win Phone, consoles)
+//  __RADMAC__ means MacOS (32 or 64-bit)
+//  __RADANDROID__ means Android NDK
+//  __RADLINUX__ means Linux (32 or 64-bit)
+//  __RADQNX__ means QNX
+//  __RADIPHONE__ means iphone (ios, tvos, or watchos)
+//  __RADIOS__ means iOS (+ __RADIPHONE__ will also be set)
+//  __RADTVOS__ means tvOS (+ __RADIPHONE__ will also be set)
+//  __RADWATCHOS__ means watchOS (+ __RADIPHONE__ will also be set)
+//  __RADEMSCRIPTEN__ means Emscripten
+
+//  __RADARM__ means arm
+//  __RADPPC__ means powerpc
+//  __RADX86__ means x86 or x64
+//  __RADX64__ means x64  (+ __RADX86__ will also be set)
+//  __RADNEON__ means you can use NEON intrinsics on ARM
+//  __RADSSE2__ means you can use SSE2 intrinsics on X86
+
+// __RADNOVARARGMACROS__ means #defines can't use ...
+
+// RADDEFSTART is "extern "C" {" on C++, nothing on C
+// RADDEFEND is "}" on C++, nothing on C
+
+// RADEXPFUNC and RADEXPLINK are used on both the declaration
+//   and definition of exported functions.
+//    RADEXPFUNC int RADEXPLINK exported_func()
+
+// RADDEFFUNC and RADLINK are used on both the declaration
+//   and definition of public functions (but not exported).
+//    RADDEFFUNC int RADLINK public_c_func()
+
+// RADRESTRICT is for non-aliasing pointers/
+
+// RADSTRUCT is defined as "struct" on msvc, and
+//   "struct __attribute__((__packed__))" on gcc/clang
+//   Used to sort of address generic structure packing
+//   (we still require #pragma packs to fix the
+//   packing on windows, though)
+
+
+// ========================================================
+// C++ name demangling nonsense
+
+#ifdef __cplusplus
+  #define RADDEFFUNC extern "C"
+  #define RADDEFSTART extern "C" {
+  #define RADDEFEND }
+  #define RADDEFAULT( val ) =val
+#else
+  #define RADDEFFUNC
+  #define RADDEFSTART
+  #define RADDEFEND
+  #define RADDEFAULT( val )
+#endif
+
+//===========================================================================
+// preprocessor string stuff
+#define RR_STRINGIZE( L )         #L
+#define RR_DO_MACRO( M, X )       M(X)
+#define RR_STRINGIZE_DELAY( X )   RR_DO_MACRO( RR_STRINGIZE, X )
+
+#define RR_STRING_JOIN(arg1, arg2)              RR_STRING_JOIN_DELAY(arg1, arg2)
+#define RR_STRING_JOIN_DELAY(arg1, arg2)        RR_STRING_JOIN_IMMEDIATE(arg1, arg2)
+#define RR_STRING_JOIN_IMMEDIATE(arg1, arg2)    arg1 ## arg2
+
+// makes "plat/plat_file"
+#define RR_PLATFORM_PATH_STR( plat, file ) RR_STRINGIZE_DELAY( RR_STRING_JOIN( plat, file ) )
+
+
+// ========================================================
+// First off, we detect your platform
+
+#ifdef __RAD_NDA_PLATFORM__
+
+  // @cdep pre $AddPlatformInclude
+  #include RR_PLATFORM_PATH_STR( __RAD_NDA_PLATFORM__, _egttypes.h )
+
+#else
+
+  #if defined(ANDROID)
+    #define __RADANDROID__ 1
+    #define __RADDETECTED__ __RADANDROID__
+  #endif
+
+  #if defined(__QNX__)
+    #define __RADQNX__ 2
+    #define __RADDETECTED__ __RADQNX__
+  #endif
+
+  #if defined(_Windows) || defined(WIN32) || defined(__WINDOWS__) || defined(_WIN32) || defined(_WIN64) || defined(WINAPI_FAMILY)
+
+    #ifdef WINAPI_FAMILY
+      // If this is #defined, we might be in a Windows Store App. But
+      // VC++ by default #defines this to a symbolic name, not an integer
+      // value, and those names are defined in "winapifamily.h". So if
+      // WINAPI_FAMILY is #defined, #include the header so we can parse it.
+      #include <winapifamily.h>
+      #define RAD_WINAPI_IS_APP (!WINAPI_FAMILY_PARTITION(WINAPI_PARTITION_DESKTOP))
+      #if RAD_WINAPI_IS_APP
+        #define __RADWINRTAPI__
+        #define __RADWINRT__ 13
+        #define __RADDETECTED__ __RADWINRT__
+      #endif
+    #else
+      #define RAD_WINAPI_IS_APP 0
+    #endif
+
+    #ifndef __RADWINRT__
+      // if we aren't WinRT, then we are plain old NT
+      #define __RADNT__ 14
+      #define __RADDETECTED__ __RADNT__
+    #endif
+  #endif
+
+  #if defined(__APPLE__)
+    #include "TargetConditionals.h"
+    #if defined(TARGET_IPHONE_SIMULATOR) && TARGET_IPHONE_SIMULATOR
+      #define __RADIPHONE__ 15
+      #define __RADIPHONESIM__ 16
+      #define __RADDETECTED__ __RADIPHONESIM__
+    #elif defined(TARGET_OS_IPHONE) && TARGET_OS_IPHONE
+      #define __RADIPHONE__ 15
+      #define __RADDETECTED__ __RADIPHONE__
+    #else
+      #define __RADMAC__ 17
+      #define __RADDETECTED__ __RADMAC__
+    #endif
+    // IOS/TVOIS/WATCHOS are subsets of __RADIPHONE__
+    #if defined(TARGET_OS_IOS) && TARGET_OS_IOS
+      #define __RADIOS__ 18
+    #endif
+    #if defined(TARGET_OS_TVOS) && TARGET_OS_TVOS
+      #define __RADTVOS__ 19
+    #endif
+    #if defined(TARGET_OS_WATCHOS) && TARGET_OS_WATCHOS
+      #define __RADWATCHOS__ 20
+    #endif
+  #endif
+
+  #if defined(__EMSCRIPTEN__)
+    #define __RADEMSCRIPTEN__  22
+    #define __RADDETECTED__ __RADEMSCRIPTEN__
+  #endif
+
+  // linux is catch all for platforms when not compiling internally at rad
+  #if defined(__linux__) && !defined(ANDROID)
+    #define __RADLINUX__ 3
+    #define __RADDETECTED__ __RADLINUX__
+  #endif
+
+  #if !defined( __RADDETECTED__) && !defined( USING_EGT )
+    // this is the catch all, when you aren't building internally at rad
+    #define __RADUNKNOWN__ 99
+    #define __RADDETECTED__ __RADUNKNOWN__
+  #endif
+
+#endif
+
+
+#if !defined(__RADDETECTED__)
+  #error "radtypes.h did not detect your platform."
+#endif
+
+// ========================================================
+// Now detect some architexture stuff
+
+#define __RAD32__ // we have no non-at-least-32-bit cpus any more
+
+#if defined(__arm__) || defined( _M_ARM )
+  #define __RADARM__ 1
+  #define __RADDETECTEDPROC__ __RADARM__
+  #define __RADLITTLEENDIAN__
+  #if defined(__ARM_NEON__) || defined(__ARM_NEON)
+    #define __RADNEON__
+  #endif
+  #if !defined(__RADSOFTFP__) && !defined( __ARM_PCS_VFP )
+    #define __RADSOFTFP__
+  #endif
+#endif
+#if defined(__i386) || defined( __i386__ ) || defined( _M_IX86 ) || defined( _X86_ )
+  #define __RADX86__ 2
+  #if !defined __RADIPHONESIM__
+    // only use mmx on PC, Win, Linux - not iphone sim!
+    #define __RADMMX__
+  #endif
+  #define __RADDETECTEDPROC__ __RADX86__
+  #define __RADLITTLEENDIAN__
+#endif
+#if defined(_x86_64) || defined( __x86_64__ ) || defined( _M_X64 ) || defined( _M_AMD64 )
+  #define __RADX86__ 2
+  #define __RADX64__ 3
+  #if !defined __RADIPHONESIM__
+    #define __RADMMX__
+  #endif
+  #define __RADDETECTEDPROC__ __RADX64__
+  #define __RADLITTLEENDIAN__
+#endif
+#if defined(__powerpc) || defined( _M_PPC )
+  #define __RADPPC__ 4
+  #define __RADALTIVEC__
+  #define __RADDETECTEDPROC__ __RADPPC__
+  #define __RADBIGENDIAN__
+#endif
+#if defined( __aarch64__ ) || defined( __arm64__ )
+  #define __RADARM__ 1
+  #define __RADARM64__ 6
+  #define __RADDETECTEDPROC__ __RADARM64__
+  #define __RADLITTLEENDIAN__
+  #define __RADNEON__
+#endif
+#if defined(__EMSCRIPTEN__)
+  #define __RADDETECTEDPROC__
+  #define __RADLITTLEENDIAN__
+#endif
+
+#if !defined(__RADDETECTEDPROC__)
+  #error "egttypes.h did not detect your processor type."
+#endif
+
+#if defined(__ppc64__) || defined(__aarch64__) || defined(_M_X64) || defined(__x86_64__) || defined(__x86_64)
+  #define __RAD64__
+  #define __RAD64REGS__  // need to set this for platforms that aren't 64-bit, but have 64-bit regs (old consoles)
+#endif
+
+
+// ========================================================
+// handle exported function declarations:
+//   in anything with __RADNOEXPORTS__, RADEXPFUNC == nothing (turn off exports with this flag)
+//   in DLL, RADEXPFUNC == DLL export
+//   in EXE, RADEXPFUNC == DLL import
+//   in EXE with RADNOEXEEXPORTS, RADEXPFUNC == nothing (turn off imports in EXE with this flag)
+//   in static lib, RADEXPFUNC == nothing
+#if ( defined(__RADINSTATICLIB__) || defined(__RADNOEXPORTS__ ) || ( defined(__RADNOEXEEXPORTS__) && ( !defined(__RADINDLL__) ) && ( !defined(__RADINSTATICLIB__) ) ) )
+  // if we are in a static lib, or exports are off, or if we are in an EXE we asked for no exe exports (or imports)
+  //   then EXPFUNC is just a normal function
+  #define RADEXPFUNC RADDEFFUNC
+#else
+  // otherwise, we use import or export base on the build flag __RADINDLL__
+  #if defined(__RADINDLL__) 
+    #define RADEXPFUNC RADDEFFUNC RADDLLEXPORTDLL
+  #else
+    #define RADEXPFUNC RADDEFFUNC RADDLLIMPORTDLL
+  #endif
+#endif
+
+#if defined(__RADANDROID__)
+  #define RADRESTRICT __restrict
+  #define RADSTRUCT struct __attribute__((__packed__))
+
+  #define RADLINK
+  #define RADEXPLINK
+  #define RADDLLEXPORTDLL __attribute__((visibility("default")))
+  #define RADDLLIMPORTDLL
+#endif
+
+#if defined(__RADQNX__)
+  #define RADRESTRICT __restrict
+  #define RADSTRUCT struct __attribute__((__packed__))
+
+  #define RADLINK
+  #define RADEXPLINK RADLINK
+  #define RADDLLEXPORTDLL
+  #define RADDLLIMPORTDLL
+#endif
+
+#if defined(__RADLINUX__) || defined(__RADUNKNOWN__)
+  #define RADRESTRICT __restrict
+  #define RADSTRUCT struct __attribute__((__packed__))
+
+  #if defined(__RADX86__) && !defined(__RADX64__)
+    #define RADLINK __attribute__((cdecl))
+    #define RADEXPLINK __attribute__((cdecl))
+  #else
+    #define RADLINK
+    #define RADEXPLINK
+  #endif
+  // for linux, we assume you are building with hidden visibility,
+  //   so for RADEXPFUNC, we turn the vis back on...
+  #define RADDLLEXPORTDLL __attribute__((visibility("default")))
+  #define RADDLLIMPORTDLL
+#endif
+
+#if defined(__RADNT__)
+  #define __RADWIN__
+  #if _MSC_VER >= 1400
+    #define RADRESTRICT __restrict
+  #else
+    // vc6 and older
+    #define RADRESTRICT
+    #define __RADNOVARARGMACROS__
+  #endif
+  #define RADSTRUCT struct 
+
+  #define RADLINK __stdcall
+  #define RADEXPLINK __stdcall
+
+  #define RADDLLEXPORTDLL __declspec(dllexport)
+  #ifdef __RADX32__
+    // on weird NT DLLs built to run on Linux and Mac, no imports
+    #define RADDLLIMPORTDLL
+  #else
+    // normal win32 dll import
+    #define RADDLLIMPORTDLL __declspec(dllimport)
+  #endif
+#endif
+
+#if defined(__RADWINRT__)
+  #define __RADWIN__
+  #if defined(__RADARM__) // no non-NEON ARMs in WinRT devices (so far)
+    #define __RADNEON__
+  #endif
+  #define RADRESTRICT __restrict
+  #define RADSTRUCT struct 
+
+  #define RADLINK __stdcall
+  #define RADEXPLINK __stdcall
+  #define RADDLLEXPORTDLL __declspec(dllexport)
+  #define RADDLLIMPORTDLL __declspec(dllimport)
+#endif
+
+#if defined(__RADIPHONE__)
+  #define __RADMACAPI__
+  #define RADRESTRICT __restrict
+  #define RADSTRUCT struct __attribute__((__packed__))
+
+  #define RADLINK
+  #define RADEXPLINK
+  #define RADDLLEXPORTDLL
+  #define RADDLLIMPORTDLL
+#endif
+
+#if defined(__RADMAC__)
+  #define __RADMACH__
+  #define __RADMACAPI__
+  #define RADRESTRICT __restrict
+  #define RADSTRUCT struct __attribute__((__packed__))
+
+  #define RADLINK
+  #define RADEXPLINK
+  // for mac, we assume you are building with hidden visibility,
+  //   so for RADEXPFUNC, we turn the vis back on...
+  #define RADDLLEXPORTDLL __attribute__((visibility("default")))
+  #define RADDLLIMPORTDLL
+
+  #ifdef TARGET_API_MAC_CARBON
+    #if TARGET_API_MAC_CARBON
+      #ifndef __RADCARBON__
+        #define __RADCARBON__
+      #endif
+    #endif
+  #endif
+#endif
+
+#if defined(__RADEMSCRIPTEN__)
+  #include <emscripten.h>
+  #define RADRESTRICT __restrict
+  #define RADSTRUCT struct __attribute__((__packed__))
+
+  #define RADLINK
+  #define RADEXPLINK EMSCRIPTEN_KEEPALIVE
+  #define RADDLLEXPORTDLL 
+  #define RADDLLIMPORTDLL
+#endif
+
+#ifndef RADLINK
+  #error RADLINK was not defined.
+#endif
+
+#ifdef _MSC_VER
+  #define RADINLINE __inline
+#else
+  #define RADINLINE inline
+#endif
+
+//===========================================================================
+// RR_NUMBERNAME is a macro to make a name unique, so that you can use it to declare
+//    variable names and they won't conflict with each other
+// using __LINE__ is broken in MSVC with /ZI , but __COUNTER__ is an MSVC extension that works
+
+#ifdef _MSC_VER
+  #define RR_NUMBERNAME(name) RR_STRING_JOIN(name,__COUNTER__)
+#else
+  #define RR_NUMBERNAME(name) RR_STRING_JOIN(name,__LINE__)
+#endif
+
+//===================================================================
+// simple compiler assert
+// this happens at declaration time, so if it's inside a function in a C file, drop {} around it
+#ifndef RR_COMPILER_ASSERT
+  #if defined(__clang__)
+    #define RR_COMPILER_ASSERT_UNUSED __attribute__((unused))  // hides warnings when compiler_asserts are in a local scope
+  #else
+    #define RR_COMPILER_ASSERT_UNUSED
+  #endif
+
+  #define RR_COMPILER_ASSERT(exp)   typedef char RR_NUMBERNAME(_dummy_array) [ (exp) ? 1 : -1 ] RR_COMPILER_ASSERT_UNUSED
+#endif
+
+
+//===========================================
+// first, we set defines for each of the types
+
+#define RAD_S8 signed char
+#define RAD_U8 unsigned char
+#define RAD_U16 unsigned short
+#define RAD_S16 signed short
+
+#if defined(__RAD64__) 
+  #define RAD_U32 unsigned int
+  #define RAD_S32 signed int
+
+  // pointers are 64 bits.
+  #if ( defined(_MSC_VER) && _MSC_VER >= 1300 && defined(_Wp64) && _Wp64 )
+    #define RAD_SINTa __w64 signed __int64
+    #define RAD_UINTa __w64 unsigned __int64
+  #else 
+    // non-vc.net compiler or /Wp64 turned off
+    #define RAD_UINTa unsigned long long
+    #define RAD_SINTa signed long long
+  #endif
+#endif
+
+#if defined(__RAD32__) && !defined(__RAD64__)
+  #define RAD_U32 unsigned int
+  #define RAD_S32 signed int
+
+  #if ( ( defined(_MSC_VER) && (_MSC_VER >= 1300 ) ) && ( defined(_Wp64) && ( _Wp64 ) ) )
+    #define RAD_SINTa __w64 signed long
+    #define RAD_UINTa __w64 unsigned long
+  #else 
+    // non-vc.net compiler or /Wp64 turned off
+    #ifdef _Wp64
+      #define RAD_SINTa signed long
+      #define RAD_UINTa unsigned long
+    #else
+      #define RAD_SINTa signed int
+      #define RAD_UINTa unsigned int
+    #endif
+  #endif
+#endif
+
+#define RAD_F32 float
+#define RAD_F64 double
+
+#if defined(_MSC_VER)
+  #define RAD_U64 unsigned __int64
+  #define RAD_S64 signed __int64
+#else
+  #define RAD_U64 unsigned long long
+  #define RAD_S64 signed long long
+#endif
+
+
+//================================================================
+// Then, we either typedef or define them based on switch settings
+
+#if !defined(RADNOTYPEDEFS)  // this define will turn off typedefs
+
+  #ifndef S8_DEFINED
+  #define S8_DEFINED
+  typedef RAD_S8 S8;
+  #endif
+
+  #ifndef U8_DEFINED
+  #define U8_DEFINED
+  typedef RAD_U8 U8;
+  #endif
+
+  #ifndef S16_DEFINED
+  #define S16_DEFINED
+  typedef RAD_S16 S16;
+  #endif
+
+  #ifndef U16_DEFINED
+  #define U16_DEFINED
+  typedef RAD_U16 U16;
+  #endif
+
+  #ifndef S32_DEFINED
+  #define S32_DEFINED
+  typedef RAD_S32 S32;
+  #endif
+
+  #ifndef U32_DEFINED
+  #define U32_DEFINED
+  typedef RAD_U32 U32;
+  #endif
+
+  #ifndef S64_DEFINED
+  #define S64_DEFINED
+  typedef RAD_S64 S64;
+  #endif
+
+  #ifndef U64_DEFINED
+  #define U64_DEFINED
+  typedef RAD_U64 U64;
+  #endif
+
+  #ifndef F32_DEFINED
+  #define F32_DEFINED
+  typedef RAD_F32 F32;
+  #endif
+
+  #ifndef F64_DEFINED
+  #define F64_DEFINED
+  typedef RAD_F64 F64;
+  #endif
+
+  #ifndef SINTa_DEFINED
+  #define SINTa_DEFINED
+  typedef RAD_SINTa SINTa;
+  #endif
+
+  #ifndef UINTa_DEFINED
+  #define UINTa_DEFINED
+  typedef RAD_UINTa UINTa;
+  #endif
+
+  #ifndef RRBOOL_DEFINED
+    #define RRBOOL_DEFINED
+    typedef S32 rrbool;
+    typedef S32 RRBOOL;
+  #endif
+
+#elif !defined(RADNOTYPEDEFINES)  // this define will turn off type defines
+
+  #ifndef S8_DEFINED
+  #define S8_DEFINED
+  #define S8 RAD_S8
+  #endif
+
+  #ifndef U8_DEFINED
+  #define U8_DEFINED
+  #define U8 RAD_U8
+  #endif
+
+  #ifndef S16_DEFINED
+  #define S16_DEFINED
+  #define S16 RAD_S16
+  #endif
+
+  #ifndef U16_DEFINED
+  #define U16_DEFINED
+  #define U16 RAD_U16
+  #endif
+
+  #ifndef S32_DEFINED
+  #define S32_DEFINED
+  #define S32 RAD_S32
+  #endif
+
+  #ifndef U32_DEFINED
+  #define U32_DEFINED
+  #define U32 RAD_U32
+  #endif
+
+  #ifndef S64_DEFINED
+  #define S64_DEFINED
+  #define S64 RAD_S64
+  #endif
+
+  #ifndef U64_DEFINED
+  #define U64_DEFINED
+  #define U64 RAD_U64
+  #endif
+
+  #ifndef F32_DEFINED
+  #define F32_DEFINED
+  #define F32 RAD_F32
+  #endif
+
+  #ifndef F64_DEFINED
+  #define F64_DEFINED
+  #define F64 RAD_F64
+  #endif
+
+  #ifndef SINTa_DEFINED
+  #define SINTa_DEFINED
+  #define SINTa RAD_SINTa
+  #endif
+
+  #ifndef UINTa_DEFINED
+  #define UINTa_DEFINED
+  #define UINTa RAD_UINTa
+  #endif
+
+  #ifndef RRBOOL_DEFINED
+    #define RRBOOL_DEFINED
+    #define rrbool S32
+    #define RRBOOL S32
+  #endif
+
+#endif
+
+#endif // __RADRES__
+
+#endif // __RADTYPESH__, __RADRR_COREH__
+
+#endif // __EGTTYPESH__
diff --git a/Engine/Plugins/Media/BinkMedia/Source/SDK/include/binkplugin.h b/Engine/Plugins/Media/BinkMedia/Source/SDK/include/binkplugin.h
deleted file mode 100644
diff --git a/Engine/Plugins/Media/BinkMedia/Source/SDK/include/binktextures.h b/Engine/Plugins/Media/BinkMedia/Source/SDK/include/binktextures.h
deleted file mode 100644
diff --git a/Engine/Plugins/Media/BinkMedia/Source/SDK/include/binktiny.h b/Engine/Plugins/Media/BinkMedia/Source/SDK/include/binktiny.h
deleted file mode 100644
diff --git a/Engine/Plugins/Media/BinkMedia/Source/SDK/include/egttypes.h b/Engine/Plugins/Media/BinkMedia/Source/SDK/include/egttypes.h
deleted file mode 100644
