Re-enabled linux editor for bink

diff --git a/Engine/Plugins/Media/BinkMedia/BinkMedia.uplugin b/Engine/Plugins/Media/BinkMedia/BinkMedia.uplugin
--- a/Engine/Plugins/Media/BinkMedia/BinkMedia.uplugin
+++ b/Engine/Plugins/Media/BinkMedia/BinkMedia.uplugin
@@ -31,7 +31,6 @@
 			"Type" : "Editor",
 			"TargetDenyList": [ "Server" ],
 			"PlatformDenyList": [ 
-				"Linux"
 			]
 		}
     ]
