Bink 2022.05 binaries

diff --git a/Engine/Plugins/Media/BinkMedia/Source/BinkMediaPlayerSDK/include/binkplugin.h b/Engine/Plugins/Media/BinkMedia/Source/BinkMediaPlayerSDK/include/binkplugin.h
--- a/Engine/Plugins/Media/BinkMedia/Source/BinkMediaPlayerSDK/include/binkplugin.h
+++ b/Engine/Plugins/Media/BinkMedia/Source/BinkMediaPlayerSDK/include/binkplugin.h
@@ -246,6 +246,14 @@ PLUG_IN_FUNC_DEC( void ) BinkPluginWindowsD3D9EndReset( void );
 
 PLUG_IN_FUNC_DEC( S32 ) BinkPluginGetPlatformInfo( U32 bink_platform_enum, void * output_ptr );
 
+// attach a subtitle file to a bink
+PLUG_IN_FUNC_DEC( S32 ) BinkPluginLoadSubtitles( BINKPLUGIN * bink, char const * srt_name );
+// attach a subtitle file to a bink, filename is UTF-16 - subtitle file contents are still UTF8 or ascii
+PLUG_IN_FUNC_DEC( S32 ) BinkPluginLoadSubtitlesUTF16( BINKPLUGIN * bink, unsigned short const * srt_name );
+// get current subtitle
+PLUG_IN_FUNC_DEC( char const * ) BinkPluginCurrentSubtitle( BINKPLUGIN * bink, U32 * iterate );
+
+
 /// Bink Image API ///
 
 typedef struct BINKIMAGE  BINKIMAGE;
@@ -325,6 +333,9 @@ PLUG_IN_FUNC_DEC( void ) BinkPluginSetPerFrameInfo( BINKPLUGINFRAMEINFO * info )
   ProcessProc(  8, void,           BinkPluginUnload,                char const * name, U64 file_byte_offset ) \
   ProcessProc(  8, S32,            BinkPluginPreloadUTF16,          unsigned short const * name, U64 file_byte_offset ) \
   ProcessProc(  8, void,           BinkPluginUnloadUTF16,           unsigned short const * name, U64 file_byte_offset ) \
+  ProcessProc(  8, S32,            BinkPluginLoadSubtitles,         BINKPLUGIN * bnk, char const * srt_name ) \
+  ProcessProc(  8, S32,            BinkPluginLoadSubtitlesUTF16,    BINKPLUGIN * bnk, unsigned short const * srt_name ) \
+  ProcessProc(  8, char const *,   BinkPluginCurrentSubtitle,       BINKPLUGIN * bnk, U32 * iterate ) \
 
 #ifdef BUILDING_FOR_UNREAL_ONLY
 
diff --git a/Engine/Plugins/Media/BinkMedia/Source/BinkMediaPlayerSDK/include/binktiny.h b/Engine/Plugins/Media/BinkMedia/Source/BinkMediaPlayerSDK/include/binktiny.h
--- a/Engine/Plugins/Media/BinkMedia/Source/BinkMediaPlayerSDK/include/binktiny.h
+++ b/Engine/Plugins/Media/BinkMedia/Source/BinkMediaPlayerSDK/include/binktiny.h
@@ -1,3 +1,7 @@
+// Copyright Epic Games, Inc. All Rights Reserved.
+// Licensed under the terms of a valid Unreal Engine license agreement,
+//   and the separate 'Unreal Engine End User License Agreement for Publishing'.
+
 #ifndef __BINKH__
 #define __BINKH__
 
