Bink: Fixes for include tool and stuff (benign changes)

diff --git a/Engine/Plugins/Media/BinkMedia/Source/BinkMediaPlayer/Private/binktexturesRHI.cpp b/Engine/Plugins/Media/BinkMedia/Source/BinkMediaPlayer/Private/binktexturesRHI.cpp
--- a/Engine/Plugins/Media/BinkMedia/Source/BinkMediaPlayer/Private/binktexturesRHI.cpp
+++ b/Engine/Plugins/Media/BinkMedia/Source/BinkMediaPlayer/Private/binktexturesRHI.cpp
@@ -1,6 +1,7 @@
 // Copyright Epic Games, Inc. All Rights Reserved.
 
 #define BINKRHIFUNCTIONS
+#define BINKTEXTURESCLEANUP
 #include "../../BinkMediaPlayerSDK/include/egttypes.h"
 #include "../../BinkMediaPlayerSDK/include/binktiny.h"
 #include "../../BinkMediaPlayerSDK/include/binktextures.h"
diff --git a/Engine/Plugins/Media/BinkMedia/Source/BinkMediaPlayerSDK/include/binktextures.h b/Engine/Plugins/Media/BinkMedia/Source/BinkMediaPlayerSDK/include/binktextures.h
--- a/Engine/Plugins/Media/BinkMedia/Source/BinkMediaPlayerSDK/include/binktextures.h
+++ b/Engine/Plugins/Media/BinkMedia/Source/BinkMediaPlayerSDK/include/binktextures.h
@@ -1,7 +1,3 @@
-// Copyright Epic Games, Inc. All Rights Reserved.
-// Licensed under the terms of a valid Unreal Engine license agreement,
-//   and the separate 'Unreal Engine End User License Agreement for Publishing'.
-
 /*
 This is Epics's high level API for using 3D hardware to do Bink GPU-assisted
 decoding and color conversion. It is supported on every 3D platform that 
@@ -136,8 +132,8 @@ RADDEFSTART
 
 #endif
 
-typedef struct BINKSHADERS BINKSHADERS;   // defined below
-typedef struct BINKTEXTURES BINKTEXTURES; // defined below
+typedef struct BINKSHADERS2 BINKSHADERS;   // defined below
+typedef struct BINKTEXTURES2 BINKTEXTURES; // defined below
 
 //=============================================================================
 
@@ -396,7 +392,7 @@ typedef S32  After_Reset_Bink_textures_ft( BINKTEXTURES * textures );
 #define After_Reset_Bink_textures( textures )            {if ((textures)->After_Reset_textures) (textures)->After_Reset_textures( textures );}
 #define Free_Bink_textures( textures )                   (textures)->Free_textures( textures )
 
-struct BINKSHADERS
+struct BINKSHADERS2
 {
   Create_Bink_textures_ft * Create_textures;
   Free_Bink_shaders_ft * Free_shaders;
@@ -405,7 +401,7 @@ struct BINKSHADERS
   UINTa user_data[ 4 ];
 };
 
-struct BINKTEXTURES
+struct BINKTEXTURES2
 {
   // wrap around doframe
   Start_Bink_texture_update_ft * Start_texture_update;
