Bink: Fix for localization jira
#p4v-preflight-copy 20149073

diff --git a/Engine/Plugins/Media/BinkMedia/Source/BinkMediaPlayer/Private/BinkMediaPlayerModule.cpp b/Engine/Plugins/Media/BinkMedia/Source/BinkMediaPlayer/Private/BinkMediaPlayerModule.cpp
--- a/Engine/Plugins/Media/BinkMedia/Source/BinkMediaPlayer/Private/BinkMediaPlayerModule.cpp
+++ b/Engine/Plugins/Media/BinkMedia/Source/BinkMediaPlayer/Private/BinkMediaPlayerModule.cpp
@@ -10,6 +10,8 @@
 
 DEFINE_LOG_CATEGORY(LogBink);
 
+#define LOCTEXT_NAMESPACE "BinkMediaPlayerModule"
+
 TSharedPtr<FBinkMovieStreamer, ESPMode::ThreadSafe> MovieStreamer;
 
 TArray< FTexture2DRHIRef > BinkActiveTextureRefs;
@@ -270,3 +272,4 @@ struct FBinkMediaPlayerModule : IModuleInterface, FTickableGameObject
 
 IMPLEMENT_MODULE( FBinkMediaPlayerModule, BinkMediaPlayer )
 
+#undef LOCTEXT_NAMESPACE
diff --git a/Engine/Plugins/Media/BinkMedia/Source/BinkMediaPlayer/Private/BinkMediaPlayerPCH.h b/Engine/Plugins/Media/BinkMedia/Source/BinkMediaPlayer/Private/BinkMediaPlayerPCH.h
--- a/Engine/Plugins/Media/BinkMedia/Source/BinkMediaPlayer/Private/BinkMediaPlayerPCH.h
+++ b/Engine/Plugins/Media/BinkMedia/Source/BinkMediaPlayer/Private/BinkMediaPlayerPCH.h
@@ -53,8 +53,5 @@ static int GetNumSpeakers()
 
 DECLARE_LOG_CATEGORY_EXTERN(LogBink, Log, All);
 
-#ifndef LOCTEXT_NAMESPACE
-#define LOCTEXT_NAMESPACE "BinkMediaPlayerModule"
-#endif
 
 
