Bink: Fixed a missing assignment to RenderTargetTextureRHI for new Texture Creation API in UE5/Main

diff --git a/Engine/Plugins/Media/BinkMedia/Source/BinkMediaPlayer/Private/Assets/BinkMediaTextureResource.cpp b/Engine/Plugins/Media/BinkMedia/Source/BinkMediaPlayer/Private/Assets/BinkMediaTextureResource.cpp
--- a/Engine/Plugins/Media/BinkMedia/Source/BinkMediaPlayer/Private/Assets/BinkMediaTextureResource.cpp
+++ b/Engine/Plugins/Media/BinkMedia/Source/BinkMediaPlayer/Private/Assets/BinkMediaTextureResource.cpp
@@ -46,7 +46,7 @@ void FBinkMediaTextureResource::InitDynamicRHI()
 		.SetFlags(TexCreateFlags | ETextureCreateFlags::RenderTargetable | ETextureCreateFlags::ShaderResource)
 		.SetInitialState(ERHIAccess::SRVMask);
 
-	TextureRHI = RHICreateTexture(Desc);
+	TextureRHI = RenderTargetTextureRHI = RHICreateTexture(Desc);
 
 	// Don't bother updating if its not a valid video
 	if (Owner->GetSurfaceWidth() && Owner->GetSurfaceHeight()) 
