Bink: Subtitles support

diff --git a/Engine/Plugins/Media/BinkMedia/Source/BinkMediaPlayer/Private/Assets/BinkMediaPlayer.cpp b/Engine/Plugins/Media/BinkMedia/Source/BinkMediaPlayer/Private/Assets/BinkMediaPlayer.cpp
--- a/Engine/Plugins/Media/BinkMedia/Source/BinkMediaPlayer/Private/Assets/BinkMediaPlayer.cpp
+++ b/Engine/Plugins/Media/BinkMedia/Source/BinkMediaPlayer/Private/Assets/BinkMediaPlayer.cpp
@@ -13,6 +13,10 @@
 #include "Slate/SlateTextures.h"
 #include "Slate/SceneViewport.h"
 
+#include "Internationalization/Internationalization.h"
+#include "Internationalization/Culture.h"
+#include "SubtitleManager.h"
+
 //#include "Media/Public/IMediaEventSink.h" //For EMediaEvent
 
 static void Command_ShowBinks()
@@ -347,6 +351,42 @@ bool UBinkMediaPlayer::Open(const FString& Url)
 		//MediaEvent.Broadcast(EMediaEvent::MediaOpenFailed);
 		return false;
 	}
+
+	// Try to open subtitles...
+	{
+		// Determine what out current culture is, and grab the most appropriate set of subtitles for it
+		FInternationalization& Internationalization = FInternationalization::Get();
+
+		const TArray<FString> PrioritizedLanguageNames = Internationalization.GetPrioritizedCultureNames(Internationalization.GetCurrentLanguage()->GetName());
+
+		CurrentHasSubtitles = 0;
+		for (const FString& LanguageName : PrioritizedLanguageNames)
+		{
+			int32 Pos = INDEX_NONE;
+			if (Url.FindLastChar(TEXT('.'), Pos))
+			{
+				const int32 PathEndPos = Url.FindLastCharByPredicate([](TCHAR C) { return C == TEXT('/') || C == TEXT('\\'); });
+				if (PathEndPos != INDEX_NONE && PathEndPos > Pos)
+				{
+					// The dot found was part of the path rather than the name
+					Pos = INDEX_NONE;
+				}
+			}
+			FString SubtitleUrl = (Pos == INDEX_NONE ? Url : Url.Left(Pos)) + "_" + LanguageName + ".srt";
+			if (FPlatformFileManager::Get().GetPlatformFile().FileExists(*SubtitleUrl))
+			{
+#if PLATFORM_WINDOWS 
+				if (BinkPluginLoadSubtitlesUTF16(bnk, (unsigned short*)*PlatformFile.ConvertToAbsolutePathForExternalAppForRead(*SubtitleUrl)))
+#else
+				if (BinkPluginLoadSubtitles(bnk, TCHAR_TO_ANSI(*PlatformFile.ConvertToAbsolutePathForExternalAppForRead(*SubtitleUrl))))
+#endif
+				{
+					CurrentHasSubtitles = 1;
+					break;
+				}
+			}
+		}
+	}
     
 	paused = false;
 	reached_end = false;
@@ -358,6 +398,15 @@ void UBinkMediaPlayer::Close()
 {
 	if(bnk) 
 	{
+		if (CurrentHasSubtitles) 
+		{
+			CurrentHasSubtitles = 0;
+
+			// Clear the movie subtitle for this object
+			TArray<FString> StopSubtitles = { TEXT("") };
+			FSubtitleManager::GetSubtitleManager()->SetMovieSubtitle(this, StopSubtitles);
+		}
+
 		ENQUEUE_RENDER_COMMAND(BinkMediaPlayer_Close_CloseBink)([bnk=bnk](FRHICommandListImmediate& RHICmdList) 
 		{ 
 			BinkPluginClose(bnk);
@@ -435,65 +484,79 @@ void UBinkMediaPlayer::Tick(float DeltaTime)
 		reached_end = true;
 		//MediaEvent.Broadcast(EMediaEvent::PlaybackEndReached);
 	}
-	if (bnk && BinkDrawStyle != 0 && GEngine && GEngine->GameViewport) 
+	if (bnk && GEngine && GEngine->GameViewport) 
 	{
 		if (!IsPlaying() && !IsPaused()) 
 		{
 			return;
 		}
 
-		BINKPLUGININFO bpinfo = {};
-		BinkPluginInfo(bnk, &bpinfo);
-		int binkw = bpinfo.Width;
-		int binkh = bpinfo.Height;
-
 		FVector2D screenSize;
 		GEngine->GameViewport->GetViewportSize(screenSize);
+		if (CurrentHasSubtitles == 1)
+		{
+			TArray<FString> SubtitlesText;
+			unsigned i = 0;
+			while(const char *sub = BinkPluginCurrentSubtitle(bnk, &i))
+			{
+				SubtitlesText.Add((UTF8CHAR*)sub);
+			}
 
-		float ulx = BinkDestinationUpperLeft.X;
-		float uly = BinkDestinationUpperLeft.Y;
-		float lrx = BinkDestinationLowerRight.X;
-		float lry = BinkDestinationLowerRight.Y;
+			FSubtitleManager::GetSubtitleManager()->SetMovieSubtitle(this, SubtitlesText);
+		}
 
-		// figure out the x,y screencoords for all of the overlay types
-		if (BinkDrawStyle == 1/*BNK_DS_OverlayFillScreenWithAspectRatio*/) 
+		if (BinkDrawStyle != 0)
 		{
-			lrx = binkw / screenSize.X;
-			lry = binkh / screenSize.Y;
-
-			if (lrx > lry) 
+			BINKPLUGININFO bpinfo = {};
+			BinkPluginInfo(bnk, &bpinfo);
+			int binkw = bpinfo.Width;
+			int binkh = bpinfo.Height;
+
+			float ulx = BinkDestinationUpperLeft.X;
+			float uly = BinkDestinationUpperLeft.Y;
+			float lrx = BinkDestinationLowerRight.X;
+			float lry = BinkDestinationLowerRight.Y;
+
+			// figure out the x,y screencoords for all of the overlay types
+			if (BinkDrawStyle == 1/*BNK_DS_OverlayFillScreenWithAspectRatio*/) 
 			{
-				lry /= lrx;
-				lrx = 1;
+				lrx = binkw / screenSize.X;
+				lry = binkh / screenSize.Y;
+
+				if (lrx > lry) 
+				{
+					lry /= lrx;
+					lrx = 1;
+				}
+				else 
+				{
+					lrx /= lry;
+					lry = 1;
+				}
+				ulx = (1.0f - lrx) / 2.0f;
+				uly = (1.0f - lry) / 2.0f;
+				lrx += ulx;
+				lry += uly;
 			}
-			else 
+			else if (BinkDrawStyle == 2/*BNK_DS_OverlayOriginalMovieSize*/) 
 			{
-				lrx /= lry;
-				lry = 1;
+				ulx = (screenSize.X - binkw) / (2.0f * screenSize.X);
+				uly = (screenSize.Y - binkh) / (2.0f * screenSize.Y);
+				lrx = binkw / screenSize.X + ulx;
+				lry = binkh / screenSize.Y + uly;
 			}
-			ulx = (1.0f - lrx) / 2.0f;
-			uly = (1.0f - lry) / 2.0f;
-			lrx += ulx;
-			lry += uly;
-		}
-		else if (BinkDrawStyle == 2/*BNK_DS_OverlayOriginalMovieSize*/) 
-		{
-			ulx = (screenSize.X - binkw) / (2.0f * screenSize.X);
-			uly = (screenSize.Y - binkh) / (2.0f * screenSize.Y);
-			lrx = binkw / screenSize.X + ulx;
-			lry = binkh / screenSize.Y + uly;
-		}
 
 #if PLATFORM_ANDROID
-		uly = 1 - uly;
-		lry = 1 - lry;
+			uly = 1 - uly;
+			lry = 1 - lry;
 #endif
 
-		ENQUEUE_RENDER_COMMAND(BinkScheduleOverlay)([bnk=bnk,ulx,uly,lrx,lry](FRHICommandListImmediate& RHICmdList) 
-		{ 
-			BinkPluginSetDrawFlags(bnk, 0);
-			BinkPluginScheduleOverlay(bnk, ulx, uly, lrx, lry, 0);
-		});
+			ENQUEUE_RENDER_COMMAND(BinkScheduleOverlay)([bnk=bnk,ulx,uly,lrx,lry](FRHICommandListImmediate& RHICmdList) 
+			{ 
+				BinkPluginSetDrawFlags(bnk, 0);
+				BinkPluginScheduleOverlay(bnk, ulx, uly, lrx, lry, 0);
+			});
+		}
 	}
 }
 
diff --git a/Engine/Plugins/Media/BinkMedia/Source/BinkMediaPlayer/Public/BinkMediaPlayer.h b/Engine/Plugins/Media/BinkMedia/Source/BinkMediaPlayer/Public/BinkMediaPlayer.h
--- a/Engine/Plugins/Media/BinkMedia/Source/BinkMediaPlayer/Public/BinkMediaPlayer.h
+++ b/Engine/Plugins/Media/BinkMedia/Source/BinkMediaPlayer/Public/BinkMediaPlayer.h
@@ -444,6 +444,7 @@ public:
 	FString CurrentUrl;
 	EBinkMediaPlayerBinkDrawStyle CurrentDrawStyle;
 	int32 CurrentLayerDepth;
+	int32 CurrentHasSubtitles;
 
 	/** Holds a delegate that is executed when media has been opened or closed. */
 	FOnMediaChanged MediaChangedEvent;
diff --git a/Engine/Source/Runtime/Engine/Private/SubtitleManager.cpp b/Engine/Source/Runtime/Engine/Private/SubtitleManager.cpp
--- a/Engine/Source/Runtime/Engine/Private/SubtitleManager.cpp
+++ b/Engine/Source/Runtime/Engine/Private/SubtitleManager.cpp
@@ -459,7 +459,17 @@ void FSubtitleManager::DisplaySubtitles(FCanvas* InCanvas, FIntRect& InSubtitleR
 				FActiveSubtitle* Subtitle = ActiveSubtitles.Find(HighestPriorityID);
 				DisplaySubtitle(InCanvas, Subtitle, InSubtitleRegion, FLinearColor::White);
 			}
-			else
+			else if (ActiveMovieSubtitles.Num() > 0)
+			{
+				for (const auto& It : ActiveMovieSubtitles)
+				{
+					if (It.Value.IsValid())
+					{
+						DisplaySubtitle(InCanvas, It.Value.Get(), InSubtitleRegion, FLinearColor::White);
+					}
+				}
+			} 
+			else 
 			{
 				CurrentSubtitleHeight = 0.0f;
 			}
