[bink] Removed RuntimeDep behavior for Movies dir

diff --git a/Engine/Plugins/Media/BinkMedia/Source/BinkMediaPlayer/BinkMediaPlayer.Build.cs b/Engine/Plugins/Media/BinkMedia/Source/BinkMediaPlayer/BinkMediaPlayer.Build.cs
--- a/Engine/Plugins/Media/BinkMedia/Source/BinkMediaPlayer/BinkMediaPlayer.Build.cs
+++ b/Engine/Plugins/Media/BinkMedia/Source/BinkMediaPlayer/BinkMediaPlayer.Build.cs
@@ -38,7 +38,7 @@ public class BinkMediaPlayer : ModuleRules
             PublicDefinitions.Add("BINKPLUGIN_UE4_EDITOR=0");
         }
 
-		RuntimeDependencies.Add("$(ProjectDir)/Content/Movies/..."); // For chunked streaming
+		//RuntimeDependencies.Add("$(ProjectDir)/Content/Movies/..."); // For chunked streaming
 
 		if (Target.Platform == UnrealTargetPlatform.Android)
 		{
