Add copyright comment

diff --git a/Engine/Plugins/Media/BinkMedia/Source/BinkMediaPlayerSDK/include/binktextures.h b/Engine/Plugins/Media/BinkMedia/Source/BinkMediaPlayerSDK/include/binktextures.h
--- a/Engine/Plugins/Media/BinkMedia/Source/BinkMediaPlayerSDK/include/binktextures.h
+++ b/Engine/Plugins/Media/BinkMedia/Source/BinkMediaPlayerSDK/include/binktextures.h
@@ -1,3 +1,5 @@
+// Copyright Epic Games Tools, LLC. All Rights Reserved.
+
 /*
 This is Epics's high level API for using 3D hardware to do Bink GPU-assisted
 decoding and color conversion. It is supported on every 3D platform that 
