Bink: Reset binks on PIE

diff --git a/Engine/Plugins/Media/BinkMedia/Source/BinkMediaPlayer/Private/BinkMediaPlayerModule.cpp b/Engine/Plugins/Media/BinkMedia/Source/BinkMediaPlayer/Private/BinkMediaPlayerModule.cpp
--- a/Engine/Plugins/Media/BinkMedia/Source/BinkMediaPlayer/Private/BinkMediaPlayerModule.cpp
+++ b/Engine/Plugins/Media/BinkMedia/Source/BinkMediaPlayer/Private/BinkMediaPlayerModule.cpp
@@ -198,6 +198,9 @@ struct FBinkMediaPlayerModule : IModuleInterface, FTickableGameObject
 
 			static IMainFrameModule& MainFrameModule = FModuleManager::LoadModuleChecked<IMainFrameModule>("MainFrame");
 			MainFrameModule.OnMainFrameCreationFinished().AddRaw(this, &FBinkMediaPlayerModule::Initialize);
+		
+			FEditorDelegates::BeginPIE.AddRaw(this, &FBinkMediaPlayerModule::HandleEditorTogglePIE);
+			FEditorDelegates::EndPIE.AddRaw(this, &FBinkMediaPlayerModule::HandleEditorTogglePIE);
 		}
 #endif
 		GetMutableDefault<UBinkMoviePlayerSettings>()->LoadConfig();
@@ -213,6 +216,10 @@ struct FBinkMediaPlayerModule : IModuleInterface, FTickableGameObject
 				GEngine->GameViewport->OnDrawn().Remove(overlayHook);
 			}
 		}
+#if BINKPLUGIN_UE4_EDITOR
+		FEditorDelegates::BeginPIE.RemoveAll(this);
+		FEditorDelegates::EndPIE.RemoveAll(this);
+#endif
 	}
 
 #if BINKPLUGIN_UE4_EDITOR
@@ -222,6 +229,19 @@ struct FBinkMediaPlayerModule : IModuleInterface, FTickableGameObject
 		PropertyModule.RegisterCustomClassLayout("MoviePlayerSettings", FOnGetDetailCustomizationInstance::CreateStatic(&FBinkMoviePlayerSettingsDetails::MakeInstance));
 		PropertyModule.NotifyCustomizationModuleChanged();
 	}
+
+	void HandleEditorTogglePIE(bool bIsSimulating)
+	{
+		for (TObjectIterator<UBinkMediaPlayer> It; It; ++It)
+		{
+			UBinkMediaPlayer* Player = *It;
+			Player->Close();
+			if(Player->StartImmediately)
+			{
+				Player->InitializePlayer();
+			}
+		}
+	}
 #endif
 
 	virtual bool IsTickable() const override { return bPluginInitialized; }
diff --git a/Engine/Plugins/Media/BinkMedia/Source/BinkMediaPlayer/Public/BinkMediaPlayer.h b/Engine/Plugins/Media/BinkMedia/Source/BinkMediaPlayer/Public/BinkMediaPlayer.h
--- a/Engine/Plugins/Media/BinkMedia/Source/BinkMediaPlayer/Public/BinkMediaPlayer.h
+++ b/Engine/Plugins/Media/BinkMedia/Source/BinkMediaPlayer/Public/BinkMediaPlayer.h
@@ -211,7 +211,7 @@ public:
 	bool OpenUrl( const FString& NewUrl );
 
 	/**
-	 * Opens the specified media URL.
+	 * Closes the specified media URL.
 	 *
 	 * @param NewUrl The URL to open.
 	 * @return true on success, false otherwise.
diff --git a/Engine/Plugins/Media/BinkMedia/Source/BinkMediaPlayerEditor/Private/Factories/BinkMediaPlayerFactory.cpp b/Engine/Plugins/Media/BinkMedia/Source/BinkMediaPlayerEditor/Private/Factories/BinkMediaPlayerFactory.cpp
--- a/Engine/Plugins/Media/BinkMedia/Source/BinkMediaPlayerEditor/Private/Factories/BinkMediaPlayerFactory.cpp
+++ b/Engine/Plugins/Media/BinkMedia/Source/BinkMediaPlayerEditor/Private/Factories/BinkMediaPlayerFactory.cpp
@@ -12,7 +12,6 @@ UBinkMediaPlayerFactory::UBinkMediaPlayerFactory( const FObjectInitializer& Obje
 	bCreateNew = false;
 	bEditorImport = true;
 
-	Formats.Add(TEXT("bik;Bink Movie File"));
 	Formats.Add(TEXT("bk2;Bink 2 Movie File"));
 }
 
