Bink: Fixed bug where w/ new PIE functionality any open video editors would be broken until reopen or property change.

diff --git a/Engine/Plugins/Media/BinkMedia/Source/BinkMediaPlayerEditor/Private/BinkMediaPlayerEditorToolkit.cpp b/Engine/Plugins/Media/BinkMedia/Source/BinkMediaPlayerEditor/Private/BinkMediaPlayerEditorToolkit.cpp
--- a/Engine/Plugins/Media/BinkMedia/Source/BinkMediaPlayerEditor/Private/BinkMediaPlayerEditorToolkit.cpp
+++ b/Engine/Plugins/Media/BinkMedia/Source/BinkMediaPlayerEditor/Private/BinkMediaPlayerEditorToolkit.cpp
@@ -24,6 +24,8 @@ void FBinkMediaPlayerEditorToolkit::Initialize( UBinkMediaPlayer* InMediaPlayer,
 	MediaPlayer->SetFlags(RF_Transactional);
 	GEditor->RegisterForUndo(this);
 
+	FEditorDelegates::EndPIE.AddRaw(this, &FBinkMediaPlayerEditorToolkit::HandleEditorEndPIE);
+
 	const FBinkMediaPlayerEditorCommands& Commands = FBinkMediaPlayerEditorCommands::Get();
 
 	ToolkitCommands->MapAction(
@@ -85,6 +87,10 @@ void FBinkMediaPlayerEditorToolkit::Initialize( UBinkMediaPlayer* InMediaPlayer,
 	RegenerateMenusAndToolbars();
 }
 
+void FBinkMediaPlayerEditorToolkit::HandleEditorEndPIE(bool bIsSimulating) {
+	MediaPlayer->InitializePlayer();
+}
+
 void FBinkMediaPlayerEditorToolkit::RegisterTabSpawners( const TSharedRef<class FTabManager>& TabManagerParam ) 
 {
 	WorkspaceMenuCategory = TabManagerParam->AddLocalWorkspaceMenuCategory(LOCTEXT("WorkspaceMenu_MediaPlayerEditor", "Bink Media Player Editor"));
diff --git a/Engine/Plugins/Media/BinkMedia/Source/BinkMediaPlayerEditor/Private/BinkMediaPlayerEditorToolkit.h b/Engine/Plugins/Media/BinkMedia/Source/BinkMediaPlayerEditor/Private/BinkMediaPlayerEditorToolkit.h
--- a/Engine/Plugins/Media/BinkMedia/Source/BinkMediaPlayerEditor/Private/BinkMediaPlayerEditorToolkit.h
+++ b/Engine/Plugins/Media/BinkMedia/Source/BinkMediaPlayerEditor/Private/BinkMediaPlayerEditorToolkit.h
@@ -22,6 +22,7 @@ struct FBinkMediaPlayerEditorToolkit : public FAssetEditorToolkit, public FEdito
 		FReimportManager::Instance()->OnPostReimport().RemoveAll(this);
 
 		GEditor->UnregisterForUndo(this);
+		FEditorDelegates::EndPIE.RemoveAll(this);
 	}
 
 	void Initialize( UBinkMediaPlayer* InMediaPlayer, const EToolkitMode::Type InMode, const TSharedPtr<class IToolkitHost>& InToolkitHost );
@@ -60,6 +61,7 @@ private:
 	void HandlePlayMediaActionExecute() { MediaPlayer->Play(); }
 	bool HandleRewindMediaActionCanExecute() const { return MediaPlayer->CanPlay() && MediaPlayer->GetTime() > FTimespan::Zero(); }
 	void HandleRewindMediaActionExecute() { MediaPlayer->Rewind(); }
+	void HandleEditorEndPIE(bool bIsSimulating);
 
 	TSharedRef<SDockTab> HandleTabManagerSpawnTab( const FSpawnTabArgs& Args, FName TabIdentifier );
 	UBinkMediaPlayer* MediaPlayer;
