Removed private PCHs from BinkMedia and BinkMediaPlayerEditor plugin.

BinkMediaPlayer:
PrivatePCH: 25.71secs
SharedPCH: 6.77secs

Saved 1.49 GBs of disk space

BinkMediaPlayerEditor:
PrivatePCH: 26.14secs
SharedPCH: 5.33secs

Saved 1.51 GBs of disk space

diff --git a/Engine/Plugins/Media/BinkMedia/Source/BinkMediaPlayer/BinkMediaPlayer.Build.cs b/Engine/Plugins/Media/BinkMedia/Source/BinkMediaPlayer/BinkMediaPlayer.Build.cs
--- a/Engine/Plugins/Media/BinkMedia/Source/BinkMediaPlayer/BinkMediaPlayer.Build.cs
+++ b/Engine/Plugins/Media/BinkMedia/Source/BinkMediaPlayer/BinkMediaPlayer.Build.cs
@@ -23,6 +23,7 @@ public class BinkMediaPlayer : ModuleRules
 
 		AddEngineThirdPartyPrivateStaticDependencies(Target, "BinkMediaPlayerSDK");
 
+		PCHUsage = PCHUsageMode.UseSharedPCHs;
         PrivatePCHHeaderFile = "Private/BinkMediaPlayerPCH.h";
 
         if (Target.bBuildEditor == true)
diff --git a/Engine/Plugins/Media/BinkMedia/Source/BinkMediaPlayerEditor/BinkMediaPlayerEditor.Build.cs b/Engine/Plugins/Media/BinkMedia/Source/BinkMediaPlayerEditor/BinkMediaPlayerEditor.Build.cs
--- a/Engine/Plugins/Media/BinkMedia/Source/BinkMediaPlayerEditor/BinkMediaPlayerEditor.Build.cs
+++ b/Engine/Plugins/Media/BinkMedia/Source/BinkMediaPlayerEditor/BinkMediaPlayerEditor.Build.cs
@@ -38,6 +38,7 @@ public class BinkMediaPlayerEditor : ModuleRules
 			}
 		);
 
+		PCHUsage = PCHUsageMode.UseSharedPCHs;
         PrivatePCHHeaderFile = "Private/BinkMediaPlayerEditorPCH.h";
 
 		PublicDefinitions.Add("BUILDING_FOR_UNREAL_ONLY=1");
