Renamed private pch headers in BinkMediaPlayer and fixed -NoPCH -DisableUnity compile issues.

diff --git a/Engine/Plugins/Media/BinkMedia/Source/BinkMediaPlayer/BinkMediaPlayer.Build.cs b/Engine/Plugins/Media/BinkMedia/Source/BinkMediaPlayer/BinkMediaPlayer.Build.cs
--- a/Engine/Plugins/Media/BinkMedia/Source/BinkMediaPlayer/BinkMediaPlayer.Build.cs
+++ b/Engine/Plugins/Media/BinkMedia/Source/BinkMediaPlayer/BinkMediaPlayer.Build.cs
@@ -23,9 +23,6 @@ public class BinkMediaPlayer : ModuleRules
 
 		AddEngineThirdPartyPrivateStaticDependencies(Target, "BinkMediaPlayerSDK");
 
-		PCHUsage = PCHUsageMode.UseSharedPCHs;
-        PrivatePCHHeaderFile = "Private/BinkMediaPlayerPCH.h";
-
         if (Target.bBuildEditor == true)
         {
 			PublicDependencyModuleNames.Add("Slate");
diff --git a/Engine/Plugins/Media/BinkMedia/Source/BinkMediaPlayer/Private/Assets/BinkMediaPlayer.cpp b/Engine/Plugins/Media/BinkMedia/Source/BinkMediaPlayer/Private/Assets/BinkMediaPlayer.cpp
--- a/Engine/Plugins/Media/BinkMedia/Source/BinkMediaPlayer/Private/Assets/BinkMediaPlayer.cpp
+++ b/Engine/Plugins/Media/BinkMedia/Source/BinkMediaPlayer/Private/Assets/BinkMediaPlayer.cpp
@@ -2,8 +2,9 @@
 //   Licenced under the Unreal Engine EULA 
 
 #include "BinkMediaPlayer.h"
+
+#include "BinkMediaPlayerPrivate.h"
 #include "binkplugin.h"
-#include "BinkMediaPlayerPCH.h"
 #include "BinkMediaTexture.h"
 #include "BinkFunctionLibrary.h"
 
diff --git a/Engine/Plugins/Media/BinkMedia/Source/BinkMediaPlayer/Private/Assets/BinkMediaTexture.cpp b/Engine/Plugins/Media/BinkMedia/Source/BinkMediaPlayer/Private/Assets/BinkMediaTexture.cpp
--- a/Engine/Plugins/Media/BinkMedia/Source/BinkMediaPlayer/Private/Assets/BinkMediaTexture.cpp
+++ b/Engine/Plugins/Media/BinkMedia/Source/BinkMediaPlayer/Private/Assets/BinkMediaTexture.cpp
@@ -2,8 +2,9 @@
 //   Licenced under the Unreal Engine EULA 
 
 #include "BinkMediaTexture.h"
-#include "BinkMediaPlayerPCH.h"
+
 #include "BinkMediaPlayer.h"
+#include "BinkMediaPlayerPrivate.h"
 #include "BinkMediaTextureResource.h"
 
 UBinkMediaTexture::UBinkMediaTexture( const FObjectInitializer& ObjectInitializer )
diff --git a/Engine/Plugins/Media/BinkMedia/Source/BinkMediaPlayer/Private/Assets/BinkMediaTextureResource.cpp b/Engine/Plugins/Media/BinkMedia/Source/BinkMediaPlayer/Private/Assets/BinkMediaTextureResource.cpp
--- a/Engine/Plugins/Media/BinkMedia/Source/BinkMediaPlayer/Private/Assets/BinkMediaTextureResource.cpp
+++ b/Engine/Plugins/Media/BinkMedia/Source/BinkMediaPlayer/Private/Assets/BinkMediaTextureResource.cpp
@@ -2,7 +2,8 @@
 //   Licenced under the Unreal Engine EULA 
 
 #include "BinkMediaTextureResource.h"
-#include "BinkMediaPlayerPCH.h"
+
+#include "BinkMediaPlayerPrivate.h"
 #include "DeviceProfiles/DeviceProfileManager.h"
 #include "DeviceProfiles/DeviceProfile.h"
 #include "Runtime/Launch/Resources/Version.h"
diff --git a/Engine/Plugins/Media/BinkMedia/Source/BinkMediaPlayer/Private/BinkFunctionLibrary.cpp b/Engine/Plugins/Media/BinkMedia/Source/BinkMediaPlayer/Private/BinkFunctionLibrary.cpp
--- a/Engine/Plugins/Media/BinkMedia/Source/BinkMediaPlayer/Private/BinkFunctionLibrary.cpp
+++ b/Engine/Plugins/Media/BinkMedia/Source/BinkMediaPlayer/Private/BinkFunctionLibrary.cpp
@@ -2,8 +2,8 @@
 //   Licenced under the Unreal Engine EULA 
 
 #include "BinkFunctionLibrary.h"
-#include "BinkMediaPlayerPCH.h"
 
+#include "BinkMediaPlayerPrivate.h"
 #include "Misc/Paths.h"
 #include "Rendering/RenderingCommon.h"
 #include "RenderingThread.h"
diff --git a/Engine/Plugins/Media/BinkMedia/Source/BinkMediaPlayer/Private/BinkMediaPlayerModule.cpp b/Engine/Plugins/Media/BinkMedia/Source/BinkMediaPlayer/Private/BinkMediaPlayerModule.cpp
--- a/Engine/Plugins/Media/BinkMedia/Source/BinkMediaPlayer/Private/BinkMediaPlayerModule.cpp
+++ b/Engine/Plugins/Media/BinkMedia/Source/BinkMediaPlayer/Private/BinkMediaPlayerModule.cpp
@@ -1,9 +1,8 @@
 // Copyright Epic Games Tools LLC
 //   Licenced under the Unreal Engine EULA 
 
-#include "BinkMediaPlayerPCH.h"
 #include "BinkFunctionLibrary.h"
-
+#include "BinkMediaPlayerPrivate.h"
 #include "Runtime/Core/Public/Misc/Paths.h"
 #include "Interfaces/IPluginManager.h"
 #include "Misc/CoreDelegates.h"
diff --git a/Engine/Plugins/Media/BinkMedia/Source/BinkMediaPlayer/Private/BinkMediaPlayerPCH.h b/Engine/Plugins/Media/BinkMedia/Source/BinkMediaPlayer/Private/BinkMediaPlayerPCH.h
deleted file mode 100644
diff --git a/Engine/Plugins/Media/BinkMedia/Source/BinkMediaPlayer/Private/BinkMediaPlayerPrivate.h b/Engine/Plugins/Media/BinkMedia/Source/BinkMediaPlayer/Private/BinkMediaPlayerPrivate.h
new file mode 100644
--- /dev/null
+++ b/Engine/Plugins/Media/BinkMedia/Source/BinkMediaPlayer/Private/BinkMediaPlayerPrivate.h
@@ -0,0 +1,57 @@
+// Copyright Epic Games Tools LLC
+//   Licenced under the Unreal Engine EULA 
+#pragma once
+
+#include "binkplugin.h"
+
+#include "HAL/PlatformFileManager.h"
+#include "Modules/ModuleManager.h"
+#include "Rendering/RenderingCommon.h"
+#include "TickableObjectRenderThread.h"
+#include "RHI.h"
+#include "RHIUtilities.h"
+#include "RHIDefinitions.h"
+
+#if PLATFORM_ANDROID
+#include <android/log.h>
+#include <android_native_app_glue.h>
+#include "Android/AndroidPlatformFile.h"
+#include "Android/AndroidJNI.h"
+#include "Android/AndroidApplication.h"
+#endif
+
+#include "BinkMovieStreamer.h"
+#include "BinkMoviePlayerSettings.h"
+
+#include "binkplugin_ue4.h"
+
+#include "AudioDeviceManager.h"
+#include "AudioMixer.h"
+#include "AudioMixerDevice.h"
+
+extern BINKMEDIAPLAYER_API unsigned bink_gpu_api;
+extern BINKMEDIAPLAYER_API unsigned bink_gpu_api_hdr;
+extern BINKMEDIAPLAYER_API EPixelFormat bink_force_pixel_format;
+extern BINKMEDIAPLAYER_API FString BinkUE4CookOnTheFlyPath(FString path, const TCHAR *filename);
+extern BINKMEDIAPLAYER_API TArray< FTexture2DRHIRef > BinkActiveTextureRefs;
+
+static int GetNumSpeakers() 
+{
+	if (!GEngine || !GEngine->GetAudioDeviceManager())
+	{
+		return 2;
+	}
+	FAudioDevice *dev = FAudioDevice::GetMainAudioDevice().GetAudioDevice();
+	if (dev && dev->IsAudioMixerEnabled()) {
+		Audio::FMixerDevice *mix = static_cast<Audio::FMixerDevice*>(dev);
+		if (mix) {
+			return mix->GetNumDeviceChannels();
+		}
+	}
+	return 2; // default case... unknown just assume 2 speakers stereo setup.
+}
+
+DECLARE_LOG_CATEGORY_EXTERN(LogBink, Log, All);
+
+
+
diff --git a/Engine/Plugins/Media/BinkMedia/Source/BinkMediaPlayer/Private/BinkMoviePlayerSettings.cpp b/Engine/Plugins/Media/BinkMedia/Source/BinkMediaPlayer/Private/BinkMoviePlayerSettings.cpp
--- a/Engine/Plugins/Media/BinkMedia/Source/BinkMediaPlayer/Private/BinkMoviePlayerSettings.cpp
+++ b/Engine/Plugins/Media/BinkMedia/Source/BinkMediaPlayer/Private/BinkMoviePlayerSettings.cpp
@@ -2,7 +2,7 @@
 //   Licenced under the Unreal Engine EULA 
 
 #include "BinkMoviePlayerSettings.h"
-#include "BinkMediaPlayerPCH.h"
+#include "BinkMediaPlayerPrivate.h"
 
 UBinkMoviePlayerSettings::UBinkMoviePlayerSettings( const FObjectInitializer& ObjectInitializer )
 	: Super(ObjectInitializer)
diff --git a/Engine/Plugins/Media/BinkMedia/Source/BinkMediaPlayer/Private/BinkMovieStreamer.cpp b/Engine/Plugins/Media/BinkMedia/Source/BinkMediaPlayer/Private/BinkMovieStreamer.cpp
--- a/Engine/Plugins/Media/BinkMedia/Source/BinkMediaPlayer/Private/BinkMovieStreamer.cpp
+++ b/Engine/Plugins/Media/BinkMedia/Source/BinkMediaPlayer/Private/BinkMovieStreamer.cpp
@@ -2,8 +2,8 @@
 //   Licenced under the Unreal Engine EULA 
 
 #include "BinkMovieStreamer.h"
-#include "BinkMediaPlayerPCH.h"
 
+#include "BinkMediaPlayerPrivate.h"
 #include "Rendering/RenderingCommon.h"
 #include "Slate/SlateTextures.h"
 #include "MoviePlayer.h"
diff --git a/Engine/Plugins/Media/BinkMedia/Source/BinkMediaPlayerEditor/BinkMediaPlayerEditor.Build.cs b/Engine/Plugins/Media/BinkMedia/Source/BinkMediaPlayerEditor/BinkMediaPlayerEditor.Build.cs
--- a/Engine/Plugins/Media/BinkMedia/Source/BinkMediaPlayerEditor/BinkMediaPlayerEditor.Build.cs
+++ b/Engine/Plugins/Media/BinkMedia/Source/BinkMediaPlayerEditor/BinkMediaPlayerEditor.Build.cs
@@ -38,9 +38,6 @@ public class BinkMediaPlayerEditor : ModuleRules
 			}
 		);
 
-		PCHUsage = PCHUsageMode.UseSharedPCHs;
-        PrivatePCHHeaderFile = "Private/BinkMediaPlayerEditorPCH.h";
-
 		PublicDefinitions.Add("BUILDING_FOR_UNREAL_ONLY=1");
 		PublicDefinitions.Add("__RADNOEXPORTS__=1");
 		PublicDefinitions.Add("__RADINSTATICLIB__=1");
diff --git a/Engine/Plugins/Media/BinkMedia/Source/BinkMediaPlayerEditor/Private/BinkMediaPlayerCustomization.cpp b/Engine/Plugins/Media/BinkMedia/Source/BinkMediaPlayerEditor/Private/BinkMediaPlayerCustomization.cpp
--- a/Engine/Plugins/Media/BinkMedia/Source/BinkMediaPlayerEditor/Private/BinkMediaPlayerCustomization.cpp
+++ b/Engine/Plugins/Media/BinkMedia/Source/BinkMediaPlayerEditor/Private/BinkMediaPlayerCustomization.cpp
@@ -2,8 +2,9 @@
 //   Licenced under the Unreal Engine EULA 
 
 #include "BinkMediaPlayerCustomization.h"
-#include "BinkMediaPlayerEditorPCH.h"
+
 #include "BinkMediaPlayer.h"
+#include "BinkMediaPlayerEditorPrivate.h"
 #include "DesktopPlatformModule.h"
 #include "Widgets/Images/SImage.h"
 #include "EditorDirectories.h"
diff --git a/Engine/Plugins/Media/BinkMedia/Source/BinkMediaPlayerEditor/Private/BinkMediaPlayerEditorModule.cpp b/Engine/Plugins/Media/BinkMedia/Source/BinkMediaPlayerEditor/Private/BinkMediaPlayerEditorModule.cpp
--- a/Engine/Plugins/Media/BinkMedia/Source/BinkMediaPlayerEditor/Private/BinkMediaPlayerEditorModule.cpp
+++ b/Engine/Plugins/Media/BinkMedia/Source/BinkMediaPlayerEditor/Private/BinkMediaPlayerEditorModule.cpp
@@ -2,8 +2,8 @@
 //   Licenced under the Unreal Engine EULA 
 
 
-#include "BinkMediaPlayerEditorPCH.h"
 #include "AssetTypeActions_Base.h"
+#include "BinkMediaPlayerEditorPrivate.h"
 #include "TextureEditor.h"
 #include "Editor/ContentBrowser/Public/ContentBrowserModule.h"
 #include "Modules/ModuleInterface.h"
diff --git a/Engine/Plugins/Media/BinkMedia/Source/BinkMediaPlayerEditor/Private/BinkMediaPlayerEditorPCH.h b/Engine/Plugins/Media/BinkMedia/Source/BinkMediaPlayerEditor/Private/BinkMediaPlayerEditorPCH.h
deleted file mode 100644
diff --git a/Engine/Plugins/Media/BinkMedia/Source/BinkMediaPlayerEditor/Private/BinkMediaPlayerEditorPrivate.h b/Engine/Plugins/Media/BinkMedia/Source/BinkMediaPlayerEditor/Private/BinkMediaPlayerEditorPrivate.h
new file mode 100644
--- /dev/null
+++ b/Engine/Plugins/Media/BinkMedia/Source/BinkMediaPlayerEditor/Private/BinkMediaPlayerEditorPrivate.h
@@ -0,0 +1,39 @@
+// Copyright Epic Games Tools LLC
+//   Licenced under the Unreal Engine EULA 
+
+#pragma once
+
+
+/* Private dependencies
+ *****************************************************************************/
+
+#include "Runtime/Launch/Resources/Version.h"
+#include "CoreMinimal.h"
+#include "AssetToolsModule.h"
+#include "Styling/SlateStyle.h"
+#include "BinkMediaPlayer.h"
+#include "BinkMediaTexture.h"
+#include "Modules/ModuleManager.h"
+#include "PropertyEditorModule.h"
+#include "Slate/SceneViewport.h"
+#include "Slate/SlateTextures.h"
+#include "TickableObjectRenderThread.h"
+#include "WorkspaceMenuStructureModule.h"
+#include "Toolkits/AssetEditorToolkit.h"
+#include "Toolkits/IToolkit.h"
+#include "PropertyCustomizationHelpers.h"
+
+/* Private includes
+ *****************************************************************************/
+
+#include "Factories/BinkMediaPlayerFactory.h"
+#include "Factories/BinkMediaPlayerFactoryNew.h"
+#include "Factories/BinkMediaTextureFactoryNew.h"
+
+#include "Models/BinkMediaPlayerEditorCommands.h"
+#include "Models/BinkMediaPlayerEditorTexture.h"
+#include "BinkMediaPlayerEditorToolkit.h"
+
+#include "SBinkWidgets.h"
+
+#include "binkplugin_ue4.h"
diff --git a/Engine/Plugins/Media/BinkMedia/Source/BinkMediaPlayerEditor/Private/BinkMediaPlayerEditorToolkit.cpp b/Engine/Plugins/Media/BinkMedia/Source/BinkMediaPlayerEditor/Private/BinkMediaPlayerEditorToolkit.cpp
--- a/Engine/Plugins/Media/BinkMedia/Source/BinkMediaPlayerEditor/Private/BinkMediaPlayerEditorToolkit.cpp
+++ b/Engine/Plugins/Media/BinkMedia/Source/BinkMediaPlayerEditor/Private/BinkMediaPlayerEditorToolkit.cpp
@@ -2,7 +2,8 @@
 //   Licenced under the Unreal Engine EULA 
 
 #include "BinkMediaPlayerEditorToolkit.h"
-#include "BinkMediaPlayerEditorPCH.h"
+
+#include "BinkMediaPlayerEditorPrivate.h"
 #include "Factories.h"
 #include "Framework/MultiBox/MultiBoxBuilder.h"
 #include "Widgets/Docking/SDockTab.h"
diff --git a/Engine/Plugins/Media/BinkMedia/Source/BinkMediaPlayerEditor/Private/Factories/BinkMediaPlayerFactory.cpp b/Engine/Plugins/Media/BinkMedia/Source/BinkMediaPlayerEditor/Private/Factories/BinkMediaPlayerFactory.cpp
--- a/Engine/Plugins/Media/BinkMedia/Source/BinkMediaPlayerEditor/Private/Factories/BinkMediaPlayerFactory.cpp
+++ b/Engine/Plugins/Media/BinkMedia/Source/BinkMediaPlayerEditor/Private/Factories/BinkMediaPlayerFactory.cpp
@@ -2,7 +2,8 @@
 //   Licenced under the Unreal Engine EULA 
 
 #include "BinkMediaPlayerFactory.h"
-#include "BinkMediaPlayerEditorPCH.h"
+
+#include "BinkMediaPlayerEditorPrivate.h"
 
 UBinkMediaPlayerFactory::UBinkMediaPlayerFactory( const FObjectInitializer& ObjectInitializer )
 	: Super(ObjectInitializer)
diff --git a/Engine/Plugins/Media/BinkMedia/Source/BinkMediaPlayerEditor/Private/Factories/BinkMediaPlayerFactoryNew.cpp b/Engine/Plugins/Media/BinkMedia/Source/BinkMediaPlayerEditor/Private/Factories/BinkMediaPlayerFactoryNew.cpp
--- a/Engine/Plugins/Media/BinkMedia/Source/BinkMediaPlayerEditor/Private/Factories/BinkMediaPlayerFactoryNew.cpp
+++ b/Engine/Plugins/Media/BinkMedia/Source/BinkMediaPlayerEditor/Private/Factories/BinkMediaPlayerFactoryNew.cpp
@@ -2,7 +2,8 @@
 //   Licenced under the Unreal Engine EULA 
 
 #include "BinkMediaPlayerFactoryNew.h"
-#include "BinkMediaPlayerEditorPCH.h"
+
+#include "BinkMediaPlayerEditorPrivate.h"
 
 UBinkMediaPlayerFactoryNew::UBinkMediaPlayerFactoryNew( const FObjectInitializer& ObjectInitializer )
 	: Super(ObjectInitializer)
diff --git a/Engine/Plugins/Media/BinkMedia/Source/BinkMediaPlayerEditor/Private/Factories/BinkMediaTextureFactoryNew.cpp b/Engine/Plugins/Media/BinkMedia/Source/BinkMediaPlayerEditor/Private/Factories/BinkMediaTextureFactoryNew.cpp
--- a/Engine/Plugins/Media/BinkMedia/Source/BinkMediaPlayerEditor/Private/Factories/BinkMediaTextureFactoryNew.cpp
+++ b/Engine/Plugins/Media/BinkMedia/Source/BinkMediaPlayerEditor/Private/Factories/BinkMediaTextureFactoryNew.cpp
@@ -2,7 +2,8 @@
 //   Licenced under the Unreal Engine EULA 
 
 #include "BinkMediaTextureFactoryNew.h"
-#include "BinkMediaPlayerEditorPCH.h"
+
+#include "BinkMediaPlayerEditorPrivate.h"
 
 UBinkMediaTextureFactoryNew::UBinkMediaTextureFactoryNew( const FObjectInitializer& ObjectInitializer )
 	: Super(ObjectInitializer)
diff --git a/Engine/Plugins/Media/BinkMedia/Source/BinkMediaPlayerEditor/Private/SBinkWidgets.cpp b/Engine/Plugins/Media/BinkMedia/Source/BinkMediaPlayerEditor/Private/SBinkWidgets.cpp
--- a/Engine/Plugins/Media/BinkMedia/Source/BinkMediaPlayerEditor/Private/SBinkWidgets.cpp
+++ b/Engine/Plugins/Media/BinkMedia/Source/BinkMediaPlayerEditor/Private/SBinkWidgets.cpp
@@ -2,7 +2,8 @@
 //   Licenced under the Unreal Engine EULA 
 
 #include "SBinkWidgets.h"
-#include "BinkMediaPlayerEditorPCH.h"
+
+#include "BinkMediaPlayerEditorPrivate.h"
 #include "Widgets/Input/SNumericEntryBox.h"
 #include "Widgets/SViewport.h"
 
