Followup to CL 22508580, fixing discovering several implementations that remain using HasActions.  More things to be upgraded.

diff --git a/Engine/Plugins/Media/BinkMedia/Source/BinkMediaPlayerEditor/Private/BinkMediaPlayerEditorModule.cpp b/Engine/Plugins/Media/BinkMedia/Source/BinkMediaPlayerEditor/Private/BinkMediaPlayerEditorModule.cpp
--- a/Engine/Plugins/Media/BinkMedia/Source/BinkMediaPlayerEditor/Private/BinkMediaPlayerEditorModule.cpp
+++ b/Engine/Plugins/Media/BinkMedia/Source/BinkMediaPlayerEditor/Private/BinkMediaPlayerEditorModule.cpp
@@ -76,6 +76,7 @@ struct FBinkMediaTextureActions : FAssetTypeActions_Base
 			)
 		);
 	}
+	
 	virtual void OpenAssetEditor( const TArray<UObject*>& InObjects, TSharedPtr<class IToolkitHost> EditWithinLevelEditor = TSharedPtr<IToolkitHost>() ) override 
 	{
 		EToolkitMode::Type Mode = EditWithinLevelEditor.IsValid() ? EToolkitMode::WorldCentric : EToolkitMode::Standalone;
@@ -165,11 +166,6 @@ struct FBinkMediaTextureActions : FAssetTypeActions_Base
 	{ 
 		return FColor::Red; 
 	}
-	virtual bool HasActions( const TArray<UObject*>& InObjects ) const override 
-	{ 
-		return true; 
-	}
-
 };
 
 struct FBinkMediaPlayerActions : FAssetTypeActions_Base 
