Fixed UE CIS compile issue.

diff --git a/Engine/Plugins/Media/BinkMedia/Source/BinkMediaPlayer/Public/BinkMediaTexture.h b/Engine/Plugins/Media/BinkMedia/Source/BinkMediaPlayer/Public/BinkMediaTexture.h
--- a/Engine/Plugins/Media/BinkMedia/Source/BinkMediaPlayer/Public/BinkMediaTexture.h
+++ b/Engine/Plugins/Media/BinkMedia/Source/BinkMediaPlayer/Public/BinkMediaTexture.h
@@ -10,7 +10,7 @@
 class UBinkMediaPlayer;
 class FBinkMediaPlayer;
 enum EPixelFormat : uint8;
-enum TextureAddress;
+enum TextureAddress : int;
 
 /**
  * Implements a texture asset for rendering video tracks from UBinkMediaPlayer assets.
