[bink] fixed build (was missing include). fix for using bink from UMG widgets, and for UMG widgets being able to call methods on Bink objects before StartupModule called.

diff --git a/Engine/Plugins/Media/BinkMedia/BinkMedia.uplugin b/Engine/Plugins/Media/BinkMedia/BinkMedia.uplugin
--- a/Engine/Plugins/Media/BinkMedia/BinkMedia.uplugin
+++ b/Engine/Plugins/Media/BinkMedia/BinkMedia.uplugin
@@ -18,7 +18,7 @@
 	[
 		{
 			"Name" : "BinkMediaPlayer",
-			"Type" : "RuntimeNoCommandlet",
+			"Type" : "Runtime",
 			"LoadingPhase" : "PreLoadingScreen",
 			"TargetDenyList": [ "Server" ],
 			"PlatformDenyList": [ 
diff --git a/Engine/Plugins/Media/BinkMedia/Source/BinkMediaPlayer/BinkMediaPlayer.Build.cs b/Engine/Plugins/Media/BinkMedia/Source/BinkMediaPlayer/BinkMediaPlayer.Build.cs
--- a/Engine/Plugins/Media/BinkMedia/Source/BinkMediaPlayer/BinkMediaPlayer.Build.cs
+++ b/Engine/Plugins/Media/BinkMedia/Source/BinkMediaPlayer/BinkMediaPlayer.Build.cs
@@ -9,7 +9,7 @@ public class BinkMediaPlayer : ModuleRules
 {
     public BinkMediaPlayer(ReadOnlyTargetRules Target) : base(Target)
     {
-		bAllowConfidentialPlatformDefines = true;
+	bAllowConfidentialPlatformDefines = true;
 
         PublicDependencyModuleNames.Add("Core");
         PublicDependencyModuleNames.Add("CoreUObject");
@@ -20,28 +20,29 @@ public class BinkMediaPlayer : ModuleRules
         PublicDependencyModuleNames.Add("MoviePlayer");
         //PublicDependencyModuleNames.Add("MediaAssets");
         PublicDependencyModuleNames.Add("Projects");
+	//PublicDependencyModuleNames.Add("PixelStreaming");
 
-		AddEngineThirdPartyPrivateStaticDependencies(Target, "BinkMediaPlayerSDK");
+	AddEngineThirdPartyPrivateStaticDependencies(Target, "BinkMediaPlayerSDK");
 
         if (Target.bBuildEditor == true)
         {
-			PublicDependencyModuleNames.Add("Slate");
-			PublicDependencyModuleNames.Add("SlateCore");
-			PublicDependencyModuleNames.Add("DesktopWidgets");
-			PublicDefinitions.Add("BINKPLUGIN_UE4_EDITOR=1");
-            PrivateDependencyModuleNames.AddRange(new string[] { "PropertyEditor", "DesktopPlatform", "SourceControl", "EditorStyle", "UnrealEd" });
+		PublicDependencyModuleNames.Add("Slate");
+		PublicDependencyModuleNames.Add("SlateCore");
+		PublicDependencyModuleNames.Add("DesktopWidgets");
+		PublicDefinitions.Add("BINKPLUGIN_UE4_EDITOR=1");
+		PrivateDependencyModuleNames.AddRange(new string[] { "PropertyEditor", "DesktopPlatform", "SourceControl", "EditorStyle", "UnrealEd" });
         }
         else
         {
-            PublicDefinitions.Add("BINKPLUGIN_UE4_EDITOR=0");
+		PublicDefinitions.Add("BINKPLUGIN_UE4_EDITOR=0");
         }
 
-		//RuntimeDependencies.Add("$(ProjectDir)/Content/Movies/..."); // For chunked streaming
+	//RuntimeDependencies.Add("$(ProjectDir)/Content/Movies/..."); // For chunked streaming
 
-		if (Target.Platform == UnrealTargetPlatform.Android)
-		{
-			string PluginPath = Utils.MakePathRelativeTo(ModuleDirectory, Target.RelativeEnginePath);
-			AdditionalPropertiesForReceipt.Add("AndroidPlugin", Path.Combine(PluginPath, "BinkMediaPlayer_APL.xml"));
-		}
+	if (Target.Platform == UnrealTargetPlatform.Android)
+	{
+		string PluginPath = Utils.MakePathRelativeTo(ModuleDirectory, Target.RelativeEnginePath);
+		AdditionalPropertiesForReceipt.Add("AndroidPlugin", Path.Combine(PluginPath, "BinkMediaPlayer_APL.xml"));
 	}
+    }
 }
diff --git a/Engine/Plugins/Media/BinkMedia/Source/BinkMediaPlayer/Private/Assets/BinkMediaPlayer.cpp b/Engine/Plugins/Media/BinkMedia/Source/BinkMediaPlayer/Private/Assets/BinkMediaPlayer.cpp
--- a/Engine/Plugins/Media/BinkMedia/Source/BinkMediaPlayer/Private/Assets/BinkMediaPlayer.cpp
+++ b/Engine/Plugins/Media/BinkMedia/Source/BinkMediaPlayer/Private/Assets/BinkMediaPlayer.cpp
@@ -282,6 +282,11 @@ void UBinkMediaPlayer::PostEditChangeProperty( FPropertyChangedEvent& PropertyCh
 
 bool UBinkMediaPlayer::Open(const FString& Url) 
 {
+	if (!BinkInitialize())
+	{
+		return false;
+	}
+
 	if (Url.IsEmpty() || IsPlaying()) 
 	{
 		return false;
diff --git a/Engine/Plugins/Media/BinkMedia/Source/BinkMediaPlayer/Private/BinkMediaPlayerModule.cpp b/Engine/Plugins/Media/BinkMedia/Source/BinkMediaPlayer/Private/BinkMediaPlayerModule.cpp
--- a/Engine/Plugins/Media/BinkMedia/Source/BinkMediaPlayer/Private/BinkMediaPlayerModule.cpp
+++ b/Engine/Plugins/Media/BinkMedia/Source/BinkMediaPlayer/Private/BinkMediaPlayerModule.cpp
@@ -121,6 +121,7 @@ struct FBinkMoviePlayerSettingsDetails : public IDetailCustomization {
 unsigned bink_gpu_api;
 unsigned bink_gpu_api_hdr;
 EPixelFormat bink_force_pixel_format = PF_Unknown;
+static bool bink_initialized = false;
 
 #ifdef BINK_NDA_GPU_ALLOC
 extern void* BinkAllocGpu(UINTa Amt, U32 Align);
@@ -141,6 +142,27 @@ static void *BinkAllocCpu(UINTa Amt, U32 Align) { return FMemory::Malloc(Amt, Al
 static void BinkFreeCpu(void * ptr) { FMemory::Free(ptr); }
 #endif
 
+bool BinkInitialize() {
+	if (bink_initialized)
+	{
+		return true;
+	}
+
+	// TODO: make this an INI setting and/or configurable in Project Settings
+	//BinkPluginTurnOnGPUAssist();
+
+	static BINKPLUGININITINFO InitInfo = { 0 };
+	InitInfo.queue = 0;
+	InitInfo.physical_device = 0;
+	InitInfo.alloc = BinkAllocCpu;
+	InitInfo.free = BinkFreeCpu;
+	InitInfo.gpu_alloc = BinkAllocGpu;
+	InitInfo.gpu_free = BinkFreeGpu;
+	bink_gpu_api = BinkRHI;
+
+	bink_initialized = (bool)BinkPluginInit(0, &InitInfo, bink_gpu_api);
+	return bink_initialized;
+}
 
 FString BinkUE4CookOnTheFlyPath(FString path, const TCHAR *filename) 
 {
@@ -162,19 +184,7 @@ struct FBinkMediaPlayerModule : IModuleInterface, FTickableGameObject
 			return;
 		}
 
-		// TODO: make this an INI setting and/or configurable in Project Settings
-		//BinkPluginTurnOnGPUAssist();
-
-		static BINKPLUGININITINFO InitInfo = { 0 };
-		InitInfo.queue = 0;
-		InitInfo.physical_device = 0;
-		InitInfo.alloc = BinkAllocCpu;
-		InitInfo.free = BinkFreeCpu;
-		InitInfo.gpu_alloc = BinkAllocGpu;
-		InitInfo.gpu_free = BinkFreeGpu;
-		bink_gpu_api = BinkRHI;
-
-		bPluginInitialized = (bool)BinkPluginInit(0, &InitInfo, bink_gpu_api);
+		bPluginInitialized = BinkInitialize();
 
 		if (!bPluginInitialized)
 		{
diff --git a/Engine/Plugins/Media/BinkMedia/Source/BinkMediaPlayer/Private/BinkMediaPlayerPrivate.h b/Engine/Plugins/Media/BinkMedia/Source/BinkMediaPlayer/Private/BinkMediaPlayerPrivate.h
--- a/Engine/Plugins/Media/BinkMedia/Source/BinkMediaPlayer/Private/BinkMediaPlayerPrivate.h
+++ b/Engine/Plugins/Media/BinkMedia/Source/BinkMediaPlayer/Private/BinkMediaPlayerPrivate.h
@@ -23,6 +23,7 @@ extern BINKMEDIAPLAYER_API unsigned bink_gpu_api_hdr;
 extern BINKMEDIAPLAYER_API EPixelFormat bink_force_pixel_format;
 extern BINKMEDIAPLAYER_API FString BinkUE4CookOnTheFlyPath(FString path, const TCHAR *filename);
 extern BINKMEDIAPLAYER_API TArray< FTexture2DRHIRef > BinkActiveTextureRefs;
+extern BINKMEDIAPLAYER_API bool BinkInitialize();
 
 static int GetNumSpeakers() 
 {
diff --git a/Engine/Plugins/Media/BinkMedia/Source/BinkMediaPlayer/Public/BinkMediaPlayer.h b/Engine/Plugins/Media/BinkMedia/Source/BinkMediaPlayer/Public/BinkMediaPlayer.h
--- a/Engine/Plugins/Media/BinkMedia/Source/BinkMediaPlayer/Public/BinkMediaPlayer.h
+++ b/Engine/Plugins/Media/BinkMedia/Source/BinkMediaPlayer/Public/BinkMediaPlayer.h
@@ -6,6 +6,7 @@
 #include "Tickable.h"
 #include "Delegates/DelegateCombinations.h"
 #include "RHI.h"
+#include "Slate/SlateTextures.h"
 #include "BinkMediaPlayer.generated.h"
 
 extern BINKMEDIAPLAYER_API unsigned bink_gpu_api;
