- Fixed TVOS Bink warning in UBT warning when building -allmodules

diff --git a/Engine/Plugins/Media/BinkMedia/Source/BinkMediaPlayer/BinkMediaPlayer.Build.cs b/Engine/Plugins/Media/BinkMedia/Source/BinkMediaPlayer/BinkMediaPlayer.Build.cs
--- a/Engine/Plugins/Media/BinkMedia/Source/BinkMediaPlayer/BinkMediaPlayer.Build.cs
+++ b/Engine/Plugins/Media/BinkMedia/Source/BinkMediaPlayer/BinkMediaPlayer.Build.cs
@@ -7,42 +7,45 @@ using UnrealBuildTool;
 
 public class BinkMediaPlayer : ModuleRules 
 {
-    public BinkMediaPlayer(ReadOnlyTargetRules Target) : base(Target)
-    {
-	bAllowConfidentialPlatformDefines = true;
-
-        PublicDependencyModuleNames.Add("Core");
-        PublicDependencyModuleNames.Add("CoreUObject");
-        PublicDependencyModuleNames.Add("Engine");
-        PublicDependencyModuleNames.Add("InputCore");
-        PublicDependencyModuleNames.Add("RenderCore");
-        PublicDependencyModuleNames.Add("RHI");
-        PublicDependencyModuleNames.Add("MoviePlayer");
-        //PublicDependencyModuleNames.Add("MediaAssets");
-        PublicDependencyModuleNames.Add("Projects");
-	//PublicDependencyModuleNames.Add("PixelStreaming");
-
-	AddEngineThirdPartyPrivateStaticDependencies(Target, "BinkMediaPlayerSDK");
-
-        if (Target.bBuildEditor == true)
-        {
-		PublicDependencyModuleNames.Add("Slate");
-		PublicDependencyModuleNames.Add("SlateCore");
-		PublicDependencyModuleNames.Add("DesktopWidgets");
-		PublicDefinitions.Add("BINKPLUGIN_UE4_EDITOR=1");
-		PrivateDependencyModuleNames.AddRange(new string[] { "PropertyEditor", "DesktopPlatform", "SourceControl", "EditorStyle", "UnrealEd" });
-        }
-        else
-        {
-		PublicDefinitions.Add("BINKPLUGIN_UE4_EDITOR=0");
-        }
-
-	//RuntimeDependencies.Add("$(ProjectDir)/Content/Movies/..."); // For chunked streaming
-
-	if (Target.Platform == UnrealTargetPlatform.Android)
+	public BinkMediaPlayer(ReadOnlyTargetRules Target) : base(Target)
 	{
-		string PluginPath = Utils.MakePathRelativeTo(ModuleDirectory, Target.RelativeEnginePath);
-		AdditionalPropertiesForReceipt.Add("AndroidPlugin", Path.Combine(PluginPath, "BinkMediaPlayer_APL.xml"));
+		bAllowConfidentialPlatformDefines = true;
+		
+		PublicDependencyModuleNames.Add("Core");
+		PublicDependencyModuleNames.Add("CoreUObject");
+		PublicDependencyModuleNames.Add("Engine");
+		PublicDependencyModuleNames.Add("InputCore");
+		PublicDependencyModuleNames.Add("RenderCore");
+		PublicDependencyModuleNames.Add("RHI");
+		PublicDependencyModuleNames.Add("MoviePlayer");
+		//PublicDependencyModuleNames.Add("MediaAssets");
+		PublicDependencyModuleNames.Add("Projects");
+		//PublicDependencyModuleNames.Add("PixelStreaming");
+		
+		if (Target.Platform != UnrealTargetPlatform.TVOS)
+		{
+			AddEngineThirdPartyPrivateStaticDependencies(Target, "BinkMediaPlayerSDK");
+		}
+		
+		if (Target.bBuildEditor == true)
+		{
+			PublicDependencyModuleNames.Add("Slate");
+			PublicDependencyModuleNames.Add("SlateCore");
+			PublicDependencyModuleNames.Add("DesktopWidgets");
+			PublicDefinitions.Add("BINKPLUGIN_UE4_EDITOR=1");
+			PrivateDependencyModuleNames.AddRange(new string[] { "PropertyEditor", "DesktopPlatform", "SourceControl", "EditorStyle", "UnrealEd" });
+		}
+		else
+		{
+			PublicDefinitions.Add("BINKPLUGIN_UE4_EDITOR=0");
+		}
+		
+		//RuntimeDependencies.Add("$(ProjectDir)/Content/Movies/..."); // For chunked streaming
+		
+		if (Target.Platform == UnrealTargetPlatform.Android)
+		{
+			string PluginPath = Utils.MakePathRelativeTo(ModuleDirectory, Target.RelativeEnginePath);
+			AdditionalPropertiesForReceipt.Add("AndroidPlugin", Path.Combine(PluginPath, "BinkMediaPlayer_APL.xml"));
+		}
 	}
-    }
 }
