Added generated.h includes and updated enums to have underlying types.

diff --git a/Engine/Plugins/AI/HTNPlanner/Source/HTNPlanner/Public/HTNDomain.h b/Engine/Plugins/AI/HTNPlanner/Source/HTNPlanner/Public/HTNDomain.h
--- a/Engine/Plugins/AI/HTNPlanner/Source/HTNPlanner/Public/HTNDomain.h
+++ b/Engine/Plugins/AI/HTNPlanner/Source/HTNPlanner/Public/HTNDomain.h
@@ -5,6 +5,8 @@
 #include "CoreMinimal.h"
 #include "UObject/ObjectMacros.h"
 
+#include "HTNDomain.generated.h"
+
 #define WITH_HTN_DEBUG !(UE_BUILD_SHIPPING || UE_BUILD_TEST)
 
 HTNPLANNER_API DECLARE_LOG_CATEGORY_EXTERN(LogHTNPlanner, Warning, All);
diff --git a/Engine/Plugins/Animation/MLDeformer/MLDeformerFramework/Source/MLDeformerFrameworkEditor/Public/MLDeformerEditorModel.h b/Engine/Plugins/Animation/MLDeformer/MLDeformerFramework/Source/MLDeformerFrameworkEditor/Public/MLDeformerEditorModel.h
--- a/Engine/Plugins/Animation/MLDeformer/MLDeformerFramework/Source/MLDeformerFrameworkEditor/Public/MLDeformerEditorModel.h
+++ b/Engine/Plugins/Animation/MLDeformer/MLDeformerFramework/Source/MLDeformerFrameworkEditor/Public/MLDeformerEditorModel.h
@@ -13,6 +13,8 @@
 #include "MLDeformerModule.h"
 #include "MLDeformerModel.h"
 
+#include "MLDeformerEditorModel.generated.h"
+
 class UMLDeformerModel;
 class UMLDeformerInputInfo;
 class AAnimationEditorPreviewActor;
diff --git a/Engine/Plugins/Animation/RigLogic/Source/RigLogicModule/Public/RigLogic.h b/Engine/Plugins/Animation/RigLogic/Source/RigLogicModule/Public/RigLogic.h
--- a/Engine/Plugins/Animation/RigLogic/Source/RigLogicModule/Public/RigLogic.h
+++ b/Engine/Plugins/Animation/RigLogic/Source/RigLogicModule/Public/RigLogic.h
@@ -11,6 +11,8 @@
 #include "FMemoryResource.h"
 #include "TransformArrayView.h"
 
+#include "RigLogic.generated.h"
+
 class FRigInstance;
 class IDNAReader;
 
diff --git a/Engine/Plugins/Editor/DisplayClusterLaunch/Source/DisplayClusterLaunchEditor/Public/DisplayClusterLaunchEditorProjectSettings.h b/Engine/Plugins/Editor/DisplayClusterLaunch/Source/DisplayClusterLaunchEditor/Public/DisplayClusterLaunchEditorProjectSettings.h
--- a/Engine/Plugins/Editor/DisplayClusterLaunch/Source/DisplayClusterLaunchEditor/Public/DisplayClusterLaunchEditorProjectSettings.h
+++ b/Engine/Plugins/Editor/DisplayClusterLaunch/Source/DisplayClusterLaunchEditor/Public/DisplayClusterLaunchEditorProjectSettings.h
@@ -7,7 +7,7 @@
 #include "DisplayClusterLaunchEditorProjectSettings.generated.h"
 
 UENUM()
-enum EDisplayClusterLaunchLogVerbosity
+enum EDisplayClusterLaunchLogVerbosity : int
 {
 	/** Always prints a fatal error to console (and log file) and crashes (even if logging is disabled) */
 	Fatal,
diff --git a/Engine/Plugins/Editor/GuidedTutorials/Source/IntroTutorials/Public/EditorTutorial.h b/Engine/Plugins/Editor/GuidedTutorials/Source/IntroTutorials/Public/EditorTutorial.h
--- a/Engine/Plugins/Editor/GuidedTutorials/Source/IntroTutorials/Public/EditorTutorial.h
+++ b/Engine/Plugins/Editor/GuidedTutorials/Source/IntroTutorials/Public/EditorTutorial.h
@@ -17,7 +17,7 @@ class UTexture2D;
 UENUM(meta=(ScriptName="TutorialContentType"))
 namespace ETutorialContent
 {
-	enum Type
+	enum Type : int
 	{
 		/** Blank - displays no content */
 		None,
@@ -37,7 +37,7 @@ namespace ETutorialContent
 UENUM()
 namespace ETutorialAnchorIdentifier
 {
-	enum Type
+	enum Type : int
 	{
 		/** No anchor */
 		None,
diff --git a/Engine/Plugins/Enterprise/DatasmithImporter/Source/DatasmithImporter/Private/DatasmithDataprepFetcher.h b/Engine/Plugins/Enterprise/DatasmithImporter/Source/DatasmithImporter/Private/DatasmithDataprepFetcher.h
--- a/Engine/Plugins/Enterprise/DatasmithImporter/Source/DatasmithImporter/Private/DatasmithDataprepFetcher.h
+++ b/Engine/Plugins/Enterprise/DatasmithImporter/Source/DatasmithImporter/Private/DatasmithDataprepFetcher.h
@@ -9,7 +9,7 @@
 #include "DatasmithDataprepFetcher.generated.h"
 
 UENUM()
-enum EMetadataKeyMatchingCriteria
+enum EMetadataKeyMatchingCriteria : int
 {
 	ExactMatch,
 	Contains
diff --git a/Engine/Plugins/Experimental/ChaosFlesh/Source/ChaosFleshNodes/Public/Dataflow/ChaosFleshTetrahedralNodes.h b/Engine/Plugins/Experimental/ChaosFlesh/Source/ChaosFleshNodes/Public/Dataflow/ChaosFleshTetrahedralNodes.h
--- a/Engine/Plugins/Experimental/ChaosFlesh/Source/ChaosFleshNodes/Public/Dataflow/ChaosFleshTetrahedralNodes.h
+++ b/Engine/Plugins/Experimental/ChaosFlesh/Source/ChaosFleshNodes/Public/Dataflow/ChaosFleshTetrahedralNodes.h
@@ -42,7 +42,7 @@ public:
 
 
 UENUM()
-enum TetMeshingMethod
+enum TetMeshingMethod : int
 {
 	IsoStuffing		UMETA(DisplayName = "IsoStuffing"),
 	TetWild			UMETA(DisplayName = "TetWild"),
diff --git a/Engine/Plugins/Experimental/CodeEditor/Source/CodeEditor/Private/CodeProjectItem.h b/Engine/Plugins/Experimental/CodeEditor/Source/CodeEditor/Private/CodeProjectItem.h
--- a/Engine/Plugins/Experimental/CodeEditor/Source/CodeEditor/Private/CodeProjectItem.h
+++ b/Engine/Plugins/Experimental/CodeEditor/Source/CodeEditor/Private/CodeProjectItem.h
@@ -8,7 +8,7 @@
 UENUM()
 namespace ECodeProjectItemType
 {
-	enum Type
+	enum Type : int
 	{
 		Project,
 		Folder,
diff --git a/Engine/Plugins/Experimental/CommonUI/Source/CommonInput/Public/CommonInputModeTypes.h b/Engine/Plugins/Experimental/CommonUI/Source/CommonInput/Public/CommonInputModeTypes.h
--- a/Engine/Plugins/Experimental/CommonUI/Source/CommonInput/Public/CommonInputModeTypes.h
+++ b/Engine/Plugins/Experimental/CommonUI/Source/CommonInput/Public/CommonInputModeTypes.h
@@ -5,6 +5,8 @@
 #include "CoreTypes.h"
 #include "UObject/ObjectMacros.h"
 
+#include "CommonInputModeTypes.generated.h"
+
 UENUM(BlueprintType)
 enum class ECommonInputMode : uint8
 {
diff --git a/Engine/Plugins/Experimental/Mutable/Source/CustomizableObjectEditor/Private/MuCOE/Nodes/SPinViewerNodeMaterialPinImageDetails.h b/Engine/Plugins/Experimental/Mutable/Source/CustomizableObjectEditor/Private/MuCOE/Nodes/SPinViewerNodeMaterialPinImageDetails.h
--- a/Engine/Plugins/Experimental/Mutable/Source/CustomizableObjectEditor/Private/MuCOE/Nodes/SPinViewerNodeMaterialPinImageDetails.h
+++ b/Engine/Plugins/Experimental/Mutable/Source/CustomizableObjectEditor/Private/MuCOE/Nodes/SPinViewerNodeMaterialPinImageDetails.h
@@ -4,6 +4,8 @@
 
 #include "MuCOE/PinViewer/SPinViewerPinDetails.h"
 
+#include "SPinViewerNodeMaterialPinImageDetails.generated.h"
+
 class UCustomizableObjectNodeMaterialPinDataImage;
 enum class EPinMode;
 namespace ESelectInfo { enum Type : int; }
diff --git a/Engine/Plugins/Experimental/NNE/Source/NNECore/Internal/NNECoreAttributeDataType.h b/Engine/Plugins/Experimental/NNE/Source/NNECore/Internal/NNECoreAttributeDataType.h
--- a/Engine/Plugins/Experimental/NNE/Source/NNECore/Internal/NNECoreAttributeDataType.h
+++ b/Engine/Plugins/Experimental/NNE/Source/NNECore/Internal/NNECoreAttributeDataType.h
@@ -5,6 +5,8 @@
 #include "UObject/Class.h"
 #include "UObject/Object.h"
 
+#include "NNECoreAttributeDataType.generated.h"
+
 /**
  * Attribute data types
  * 
diff --git a/Engine/Plugins/Experimental/NNE/Source/NNECore/Public/NNECoreTypes.h b/Engine/Plugins/Experimental/NNE/Source/NNECore/Public/NNECoreTypes.h
--- a/Engine/Plugins/Experimental/NNE/Source/NNECore/Public/NNECoreTypes.h
+++ b/Engine/Plugins/Experimental/NNE/Source/NNECore/Public/NNECoreTypes.h
@@ -5,6 +5,8 @@
 #include "CoreMinimal.h"
 #include "UObject/Class.h"
 
+#include "NNECoreTypes.generated.h"
+
 /**
  * The enum lists all tensor data types used in NNE 
  */
diff --git a/Engine/Plugins/Experimental/Text3D/Source/Text3D/Public/BevelType.h b/Engine/Plugins/Experimental/Text3D/Source/Text3D/Public/BevelType.h
--- a/Engine/Plugins/Experimental/Text3D/Source/Text3D/Public/BevelType.h
+++ b/Engine/Plugins/Experimental/Text3D/Source/Text3D/Public/BevelType.h
@@ -4,6 +4,8 @@
 
 #include "UObject/ObjectMacros.h"
 
+#include "BevelType.generated.h"
+
 UENUM()
 enum class EText3DBevelType : uint8
 {
diff --git a/Engine/Plugins/Experimental/VirtualProduction/LedWallCalibration/Source/LedWallCalibration/Public/LedWallArucoGenerationOptions.h b/Engine/Plugins/Experimental/VirtualProduction/LedWallCalibration/Source/LedWallCalibration/Public/LedWallArucoGenerationOptions.h
--- a/Engine/Plugins/Experimental/VirtualProduction/LedWallCalibration/Source/LedWallCalibration/Public/LedWallArucoGenerationOptions.h
+++ b/Engine/Plugins/Experimental/VirtualProduction/LedWallCalibration/Source/LedWallCalibration/Public/LedWallArucoGenerationOptions.h
@@ -8,7 +8,7 @@
 
 /** Enumeration of known Aruco dictionaries */
 UENUM(BlueprintType)
-enum EArucoDictionary
+enum EArucoDictionary : int
 {
 	DICT_4X4_50          UMETA(DisplayName = "DICT_4X4_50"),
 	DICT_4X4_100		 UMETA(DisplayName = "DICT_4X4_100"),
diff --git a/Engine/Plugins/Importers/USDImporter/Source/UnrealUSDWrapper/Public/UnrealUSDWrapper.h b/Engine/Plugins/Importers/USDImporter/Source/UnrealUSDWrapper/Public/UnrealUSDWrapper.h
--- a/Engine/Plugins/Importers/USDImporter/Source/UnrealUSDWrapper/Public/UnrealUSDWrapper.h
+++ b/Engine/Plugins/Importers/USDImporter/Source/UnrealUSDWrapper/Public/UnrealUSDWrapper.h
@@ -21,6 +21,8 @@
 #include "USDIncludesEnd.h"
 #endif // #if USE_USD_SDK
 
+#include "UnrealUSDWrapper.generated.h"
+
 #if USE_USD_SDK
 PXR_NAMESPACE_OPEN_SCOPE
 	class GfMatrix4d;
diff --git a/Engine/Plugins/Media/BinkMedia/Source/BinkMediaPlayer/Private/BinkMoviePlayerSettings.h b/Engine/Plugins/Media/BinkMedia/Source/BinkMediaPlayer/Private/BinkMoviePlayerSettings.h
--- a/Engine/Plugins/Media/BinkMedia/Source/BinkMediaPlayer/Private/BinkMoviePlayerSettings.h
+++ b/Engine/Plugins/Media/BinkMedia/Source/BinkMediaPlayer/Private/BinkMoviePlayerSettings.h
@@ -11,7 +11,7 @@ enum EPixelFormat : uint8;
  * Enumerates available bink buffering modes.
  */
 UENUM()
-enum EBinkMoviePlayerBinkBufferModes
+enum EBinkMoviePlayerBinkBufferModes : int
 {
 	/** Stream the movie off the media during playback (caches about 1 second of video). */
 	MP_Bink_Stream UMETA(DisplayName="Stream"),
@@ -29,7 +29,7 @@ enum EBinkMoviePlayerBinkBufferModes
  * Enumerates available used to specify the sounds to open at playback w/ bink movies.
  */
 UENUM()
-enum EBinkMoviePlayerBinkSoundTrack
+enum EBinkMoviePlayerBinkSoundTrack : int
 {
 	/** Don't open any sound tracks snd_track_start not used. */
 	MP_Bink_Sound_None UMETA(DisplayName="None"),
diff --git a/Engine/Plugins/Media/BinkMedia/Source/BinkMediaPlayer/Public/BinkMediaPlayer.h b/Engine/Plugins/Media/BinkMedia/Source/BinkMediaPlayer/Public/BinkMediaPlayer.h
--- a/Engine/Plugins/Media/BinkMedia/Source/BinkMediaPlayer/Public/BinkMediaPlayer.h
+++ b/Engine/Plugins/Media/BinkMedia/Source/BinkMediaPlayer/Public/BinkMediaPlayer.h
@@ -18,7 +18,7 @@ extern BINKMEDIAPLAYER_API FString BinkUE4CookOnTheFlyPath(FString path, const T
  * Enumerates available bink buffering modes.
  */
 UENUM()
-enum EBinkMediaPlayerBinkBufferModes
+enum EBinkMediaPlayerBinkBufferModes : int
 {
 	/** Stream the movie off the media during playback (caches about 1 second of video). */
 	BMASM_Bink_Stream UMETA(DisplayName="Stream"),
@@ -36,7 +36,7 @@ enum EBinkMediaPlayerBinkBufferModes
  * Enumerates available used to specify the sounds to open at playback w/ bink movies.
  */
 UENUM()
-enum EBinkMediaPlayerBinkSoundTrack
+enum EBinkMediaPlayerBinkSoundTrack : int
 {
 	/** Don't open any sound tracks snd_track_start not used. */
 	BMASM_Bink_Sound_None UMETA(DisplayName="None"),
@@ -66,7 +66,7 @@ enum EBinkMediaPlayerBinkSoundTrack
  * Enumerates available bink buffering modes.
  */
 UENUM()
-enum EBinkMediaPlayerBinkDrawStyle
+enum EBinkMediaPlayerBinkDrawStyle : int
 {
 	/** Renders movie to the destination texture (default UE4 functionality) */
 	BMASM_Bink_DS_RenderToTexture UMETA(DisplayName="Render to Texture"),
diff --git a/Engine/Plugins/NetcodeUnitTest/NetcodeUnitTest/Source/NetcodeUnitTest/Public/NUTEnum.h b/Engine/Plugins/NetcodeUnitTest/NetcodeUnitTest/Source/NetcodeUnitTest/Public/NUTEnum.h
--- a/Engine/Plugins/NetcodeUnitTest/NetcodeUnitTest/Source/NetcodeUnitTest/Public/NUTEnum.h
+++ b/Engine/Plugins/NetcodeUnitTest/NetcodeUnitTest/Source/NetcodeUnitTest/Public/NUTEnum.h
@@ -6,6 +6,8 @@
 
 #include "NetcodeUnitTest.h"
 
+#include "NUTEnum.generated.h"
+
 /**
  * Flags for configuring the minimal client, what parts of the netcode should be enabled etc.
  */
diff --git a/Engine/Plugins/Runtime/AR/AppleAR/AppleARKit/Source/AppleARKit/Private/ExternalTextureMaterialExpression.h b/Engine/Plugins/Runtime/AR/AppleAR/AppleARKit/Source/AppleARKit/Private/ExternalTextureMaterialExpression.h
--- a/Engine/Plugins/Runtime/AR/AppleAR/AppleARKit/Source/AppleARKit/Private/ExternalTextureMaterialExpression.h
+++ b/Engine/Plugins/Runtime/AR/AppleAR/AppleARKit/Source/AppleARKit/Private/ExternalTextureMaterialExpression.h
@@ -8,7 +8,7 @@
 #include "ExternalTextureMaterialExpression.generated.h"
 
 UENUM()
-enum EARKitTextureType
+enum EARKitTextureType : int
 {
 	TextureY,
 	TextureCbCr,
diff --git a/Engine/Plugins/Runtime/Analytics/Adjust/Source/AdjustEditor/Classes/AdjustSettings.h b/Engine/Plugins/Runtime/Analytics/Adjust/Source/AdjustEditor/Classes/AdjustSettings.h
--- a/Engine/Plugins/Runtime/Analytics/Adjust/Source/AdjustEditor/Classes/AdjustSettings.h
+++ b/Engine/Plugins/Runtime/Analytics/Adjust/Source/AdjustEditor/Classes/AdjustSettings.h
@@ -8,7 +8,7 @@
 UENUM()
 namespace EAndroidAdjustLogging
 {
-	enum Type
+	enum Type : int
 	{
 		/** enable all logging */
 		verbose  UMETA(DisplayName = "VERBOSE"),
diff --git a/Engine/Plugins/Runtime/ApexDestruction/Source/ApexDestruction/Public/DestructibleMesh.h b/Engine/Plugins/Runtime/ApexDestruction/Source/ApexDestruction/Public/DestructibleMesh.h
--- a/Engine/Plugins/Runtime/ApexDestruction/Source/ApexDestruction/Public/DestructibleMesh.h
+++ b/Engine/Plugins/Runtime/ApexDestruction/Source/ApexDestruction/Public/DestructibleMesh.h
@@ -20,7 +20,7 @@ class UStaticMesh;
 	Chunks up to the depth DefaultImpactDamageDepth will take impact damage, unless IDO_On or IDO_Off is chosen.
 */
 UENUM()
-enum EImpactDamageOverride
+enum EImpactDamageOverride : int
 {
 	IDO_None,
 	IDO_On,
diff --git a/Engine/Plugins/Runtime/MassEntity/Source/MassEntity/Public/MassCommands.h b/Engine/Plugins/Runtime/MassEntity/Source/MassEntity/Public/MassCommands.h
--- a/Engine/Plugins/Runtime/MassEntity/Source/MassEntity/Public/MassCommands.h
+++ b/Engine/Plugins/Runtime/MassEntity/Source/MassEntity/Public/MassCommands.h
@@ -8,6 +8,8 @@
 #include "MassEntityUtils.h"
 #include "MassEntityManager.h"
 
+#include "MassCommands.generated.h"
+
 /**
  * Enum used by MassBatchCommands to declare their "type". This data is later used to group commands so that command 
  * effects are applied in a controllable fashion 
diff --git a/Engine/Plugins/Runtime/MassGameplay/Source/MassLOD/Public/MassLODTypes.h b/Engine/Plugins/Runtime/MassGameplay/Source/MassLOD/Public/MassLODTypes.h
--- a/Engine/Plugins/Runtime/MassGameplay/Source/MassLOD/Public/MassLODTypes.h
+++ b/Engine/Plugins/Runtime/MassGameplay/Source/MassLOD/Public/MassLODTypes.h
@@ -35,7 +35,7 @@ namespace UE::Mass::ProcessorGroupNames
 UENUM()
 namespace EMassLOD
 {
-	enum Type
+	enum Type : int
 	{
 		High,
 		Medium,
diff --git a/Engine/Plugins/Runtime/PreLoadScreenMoviePlayer/Source/PreLoadScreenMoviePlayer/Public/MoviePlayerAttributes.h b/Engine/Plugins/Runtime/PreLoadScreenMoviePlayer/Source/PreLoadScreenMoviePlayer/Public/MoviePlayerAttributes.h
--- a/Engine/Plugins/Runtime/PreLoadScreenMoviePlayer/Source/PreLoadScreenMoviePlayer/Public/MoviePlayerAttributes.h
+++ b/Engine/Plugins/Runtime/PreLoadScreenMoviePlayer/Source/PreLoadScreenMoviePlayer/Public/MoviePlayerAttributes.h
@@ -5,8 +5,10 @@
 #include "CoreMinimal.h"
 #include "UObject/ObjectMacros.h"
 
+#include "MoviePlayerAttributes.generated.h"
+
 UENUM()
-enum EMovieScreenPlaybackType
+enum EMovieScreenPlaybackType : int
 {
     /** Normal playback mode.  Play each movie in the play list a single time */
     MT_MS_Normal UMETA(DisplayName = "Normal Playback"),
diff --git a/Engine/Plugins/Runtime/Steam/SteamAudio/Source/SteamAudio/Private/PhononCommon.h b/Engine/Plugins/Runtime/Steam/SteamAudio/Source/SteamAudio/Private/PhononCommon.h
--- a/Engine/Plugins/Runtime/Steam/SteamAudio/Source/SteamAudio/Private/PhononCommon.h
+++ b/Engine/Plugins/Runtime/Steam/SteamAudio/Source/SteamAudio/Private/PhononCommon.h
@@ -8,6 +8,8 @@
 #include "CoreMinimal.h"
 #include "UObject/ObjectMacros.h"
 
+#include "PhononCommon.generated.h"
+
 class UPhononGeometryComponent;
 
 DECLARE_LOG_CATEGORY_EXTERN(LogSteamAudio, Log, All);
diff --git a/Engine/Plugins/Runtime/Steam/SteamAudio/Source/SteamAudio/Private/PhononMaterial.h b/Engine/Plugins/Runtime/Steam/SteamAudio/Source/SteamAudio/Private/PhononMaterial.h
--- a/Engine/Plugins/Runtime/Steam/SteamAudio/Source/SteamAudio/Private/PhononMaterial.h
+++ b/Engine/Plugins/Runtime/Steam/SteamAudio/Source/SteamAudio/Private/PhononMaterial.h
@@ -8,6 +8,8 @@
 #include "UObject/ObjectMacros.h"
 #include "phonon.h"
 
+#include "PhononMaterial.generated.h"
+
 UENUM()
 enum class EPhononMaterial : uint8
 {
diff --git a/Engine/Plugins/VirtualProduction/CameraCalibrationCore/Source/CameraCalibrationCore/Public/CalibrationPointComponent.h b/Engine/Plugins/VirtualProduction/CameraCalibrationCore/Source/CameraCalibrationCore/Public/CalibrationPointComponent.h
--- a/Engine/Plugins/VirtualProduction/CameraCalibrationCore/Source/CameraCalibrationCore/Public/CalibrationPointComponent.h
+++ b/Engine/Plugins/VirtualProduction/CameraCalibrationCore/Source/CameraCalibrationCore/Public/CalibrationPointComponent.h
@@ -12,7 +12,7 @@
 
 /** Ways of visualizing the calibration points */
 UENUM(BlueprintType)
-enum ECalibrationPointVisualization
+enum ECalibrationPointVisualization : int
 {
 	CalibrationPointVisualizationCube    UMETA(DisplayName = "Cubes"),
 	CalibrationPointVisualizationPyramid UMETA(DisplayName = "Pyramids"),
diff --git a/Engine/Plugins/VirtualProduction/DMX/DMXFixtures/Source/DMXFixtures/Public/DMXFixtureActor.h b/Engine/Plugins/VirtualProduction/DMX/DMXFixtures/Source/DMXFixtures/Public/DMXFixtureActor.h
--- a/Engine/Plugins/VirtualProduction/DMX/DMXFixtures/Source/DMXFixtures/Public/DMXFixtureActor.h
+++ b/Engine/Plugins/VirtualProduction/DMX/DMXFixtures/Source/DMXFixtures/Public/DMXFixtureActor.h
@@ -21,7 +21,7 @@
 
 
 UENUM()
-enum EDMXFixtureQualityLevel
+enum EDMXFixtureQualityLevel : int
 {
 	LowQuality			UMETA(DisplayName = "Low"),
 	MediumQuality		UMETA(DisplayName = "Medium"),
diff --git a/Engine/Plugins/VirtualProduction/LevelSnapshots/Source/LevelShapshotFilters/Public/Builtin/ActorChangedTransformFilter.h b/Engine/Plugins/VirtualProduction/LevelSnapshots/Source/LevelShapshotFilters/Public/Builtin/ActorChangedTransformFilter.h
--- a/Engine/Plugins/VirtualProduction/LevelSnapshots/Source/LevelShapshotFilters/Public/Builtin/ActorChangedTransformFilter.h
+++ b/Engine/Plugins/VirtualProduction/LevelSnapshots/Source/LevelShapshotFilters/Public/Builtin/ActorChangedTransformFilter.h
@@ -9,7 +9,7 @@
 UENUM()
 namespace ETransformReturnType
 {
-	enum Type
+	enum Type : int
 	{
 		/* Return true if the snapshot and world actor have different transforms */
 		IsValidWhenTransformChanged,
diff --git a/Engine/Plugins/VirtualProduction/LevelSnapshots/Source/LevelShapshotFilters/Public/Builtin/ActorHasTagFilter.h b/Engine/Plugins/VirtualProduction/LevelSnapshots/Source/LevelShapshotFilters/Public/Builtin/ActorHasTagFilter.h
--- a/Engine/Plugins/VirtualProduction/LevelSnapshots/Source/LevelShapshotFilters/Public/Builtin/ActorHasTagFilter.h
+++ b/Engine/Plugins/VirtualProduction/LevelSnapshots/Source/LevelShapshotFilters/Public/Builtin/ActorHasTagFilter.h
@@ -9,7 +9,7 @@
 UENUM()
 namespace EActorToCheck
 {
-	enum Type
+	enum Type : int
 	{
 		/* Checks only the tags of the world actor */
 		WorldActor,
@@ -23,7 +23,7 @@ namespace EActorToCheck
 UENUM()
 namespace ETagCheckingBehavior
 {
-	enum Type
+	enum Type : int
 	{
 		/* Actor must have all tags to pass */
 		HasAllTags,
diff --git a/Engine/Plugins/VirtualProduction/LevelSnapshots/Source/LevelShapshotFilters/Public/Builtin/PropertyHasNameFilter.h b/Engine/Plugins/VirtualProduction/LevelSnapshots/Source/LevelShapshotFilters/Public/Builtin/PropertyHasNameFilter.h
--- a/Engine/Plugins/VirtualProduction/LevelSnapshots/Source/LevelShapshotFilters/Public/Builtin/PropertyHasNameFilter.h
+++ b/Engine/Plugins/VirtualProduction/LevelSnapshots/Source/LevelShapshotFilters/Public/Builtin/PropertyHasNameFilter.h
@@ -9,7 +9,7 @@
 UENUM()
 namespace ENameMatchingRule
 {
-	enum Type
+	enum Type : int
 	{
 		/* The property name must match the given name exactly. */
 		MatchesExactly,
diff --git a/Engine/Plugins/VirtualProduction/LevelSnapshots/Source/LevelShapshotFilters/Public/Builtin/PropertyTypeFilter.h b/Engine/Plugins/VirtualProduction/LevelSnapshots/Source/LevelShapshotFilters/Public/Builtin/PropertyTypeFilter.h
--- a/Engine/Plugins/VirtualProduction/LevelSnapshots/Source/LevelShapshotFilters/Public/Builtin/PropertyTypeFilter.h
+++ b/Engine/Plugins/VirtualProduction/LevelSnapshots/Source/LevelShapshotFilters/Public/Builtin/PropertyTypeFilter.h
@@ -10,7 +10,7 @@ UENUM()
 namespace EBlueprintPropertyType
 {
 	// Copied from EPropertyType
-	enum Type
+	enum Type : int
 	{
 		// Commented out unsupported types types
 		// Leave comments here so it is easier to see what was removed
diff --git a/Engine/Plugins/VirtualProduction/LevelSnapshots/Source/LevelShapshotFilters/Public/LevelSnapshotFilters.h b/Engine/Plugins/VirtualProduction/LevelSnapshots/Source/LevelShapshotFilters/Public/LevelSnapshotFilters.h
--- a/Engine/Plugins/VirtualProduction/LevelSnapshots/Source/LevelShapshotFilters/Public/LevelSnapshotFilters.h
+++ b/Engine/Plugins/VirtualProduction/LevelSnapshots/Source/LevelShapshotFilters/Public/LevelSnapshotFilters.h
@@ -10,7 +10,7 @@
 UENUM(BlueprintType)
 namespace EFilterResult
 {
-	enum Type
+	enum Type : int
 	{
 		/* This actor / property will be included. 
 		 */
diff --git a/Engine/Plugins/VirtualProduction/LevelSnapshots/Source/LevelSnapshots/Public/Settings/HashSettings.h b/Engine/Plugins/VirtualProduction/LevelSnapshots/Source/LevelSnapshots/Public/Settings/HashSettings.h
--- a/Engine/Plugins/VirtualProduction/LevelSnapshots/Source/LevelSnapshots/Public/Settings/HashSettings.h
+++ b/Engine/Plugins/VirtualProduction/LevelSnapshots/Source/LevelSnapshots/Public/Settings/HashSettings.h
@@ -8,7 +8,7 @@
 UENUM()
 namespace EHashAlgorithmChooseBehavior
 {
-	enum Type
+	enum Type : int
 	{
 		/**
 		 * Use whichever algorithm is faster
diff --git a/Engine/Plugins/VirtualProduction/RemoteControl/Source/RemoteControl/Private/Tests/RemoteControlPropertyHandleTestData.h b/Engine/Plugins/VirtualProduction/RemoteControl/Source/RemoteControl/Private/Tests/RemoteControlPropertyHandleTestData.h
--- a/Engine/Plugins/VirtualProduction/RemoteControl/Source/RemoteControl/Private/Tests/RemoteControlPropertyHandleTestData.h
+++ b/Engine/Plugins/VirtualProduction/RemoteControl/Source/RemoteControl/Private/Tests/RemoteControlPropertyHandleTestData.h
@@ -16,7 +16,7 @@ enum class ERemoteControlEnumClass : uint8
 UENUM()
 namespace ERemoteControlEnum
 {
-	enum Type
+	enum Type : int
 	{
 		E_One,
 		E_Two,
diff --git a/Engine/Plugins/VirtualProduction/RemoteControl/Source/RemoteControl/Public/IRemoteControlModule.h b/Engine/Plugins/VirtualProduction/RemoteControl/Source/RemoteControl/Public/IRemoteControlModule.h
--- a/Engine/Plugins/VirtualProduction/RemoteControl/Source/RemoteControl/Public/IRemoteControlModule.h
+++ b/Engine/Plugins/VirtualProduction/RemoteControl/Source/RemoteControl/Public/IRemoteControlModule.h
@@ -9,6 +9,8 @@
 #include "UObject/StructOnScope.h"
 #include "UObject/WeakFieldPtr.h"
 
+#include "IRemoteControlModule.generated.h"
+
 REMOTECONTROL_API DECLARE_LOG_CATEGORY_EXTERN(LogRemoteControl, Log, All);
 
 class IRemoteControlMaskingFactory;
diff --git a/Engine/Plugins/VirtualProduction/RemoteControl/Source/RemoteControlCommon/Public/RemoteControlCommon.h b/Engine/Plugins/VirtualProduction/RemoteControl/Source/RemoteControlCommon/Public/RemoteControlCommon.h
--- a/Engine/Plugins/VirtualProduction/RemoteControl/Source/RemoteControlCommon/Public/RemoteControlCommon.h
+++ b/Engine/Plugins/VirtualProduction/RemoteControl/Source/RemoteControlCommon/Public/RemoteControlCommon.h
@@ -9,7 +9,7 @@
 UENUM()
 namespace ERCProtocolBinding
 {
-	enum Op
+	enum Op : int
 	{
 		Added,
 
diff --git a/Engine/Source/Runtime/GameMenuBuilder/Public/GameMenuWidgetStyle.h b/Engine/Source/Runtime/GameMenuBuilder/Public/GameMenuWidgetStyle.h
--- a/Engine/Source/Runtime/GameMenuBuilder/Public/GameMenuWidgetStyle.h
+++ b/Engine/Source/Runtime/GameMenuBuilder/Public/GameMenuWidgetStyle.h
@@ -16,7 +16,7 @@
 UENUM()
 namespace GameMenuLayoutType
 {
-	enum Type
+	enum Type : int
 	{
 		/* Single one page menu. */
 		Single,
