[Bink] Fix for tvOS when building with allmodules. Try #2

diff --git a/Engine/Plugins/Media/BinkMedia/Source/BinkMediaPlayer/BinkMediaPlayer.Build.cs b/Engine/Plugins/Media/BinkMedia/Source/BinkMediaPlayer/BinkMediaPlayer.Build.cs
--- a/Engine/Plugins/Media/BinkMedia/Source/BinkMediaPlayer/BinkMediaPlayer.Build.cs
+++ b/Engine/Plugins/Media/BinkMedia/Source/BinkMediaPlayer/BinkMediaPlayer.Build.cs
@@ -22,10 +22,7 @@ public class BinkMediaPlayer : ModuleRules
 		PublicDependencyModuleNames.Add("Projects");
 		//PublicDependencyModuleNames.Add("PixelStreaming");
 		
-		if (Target.Platform != UnrealTargetPlatform.TVOS)
-		{
-			AddEngineThirdPartyPrivateStaticDependencies(Target, "BinkMediaPlayerSDK");
-		}
+		AddEngineThirdPartyPrivateStaticDependencies(Target, "BinkMediaPlayerSDK");
 		
 		if (Target.bBuildEditor == true)
 		{
diff --git a/Engine/Plugins/Media/BinkMedia/Source/BinkMediaPlayerSDK/BinkMediaPlayerSDK.Build.cs b/Engine/Plugins/Media/BinkMedia/Source/BinkMediaPlayerSDK/BinkMediaPlayerSDK.Build.cs
--- a/Engine/Plugins/Media/BinkMedia/Source/BinkMediaPlayerSDK/BinkMediaPlayerSDK.Build.cs
+++ b/Engine/Plugins/Media/BinkMedia/Source/BinkMediaPlayerSDK/BinkMediaPlayerSDK.Build.cs
@@ -67,12 +67,17 @@ public class BinkMediaPlayerSDK : ModuleRules
 				PublicDependencyModuleNames.Add("MetalRHI");
 				PublicFrameworks.Add("AudioToolbox");
 			}
-            else if (Target.Platform == UnrealTargetPlatform.IOS)
-            {
-                Lib = "libBinkUnreal" + Platform + ".a";
-                PublicDependencyModuleNames.Add("MetalRHI");
-            }
-            else if (Target.IsInPlatformGroup(UnrealPlatformGroup.Apple))
+			else if (Target.Platform == UnrealTargetPlatform.IOS)
+			{
+				Lib = "libBinkUnreal" + Platform + ".a";
+				PublicDependencyModuleNames.Add("MetalRHI");
+			}
+			else if (Target.Platform == UnrealTargetPlatform.TVOS)
+			{
+				Lib = null; // Unsupported platform
+				PublicDependencyModuleNames.Add("MetalRHI");
+			}
+			else if (Target.IsInPlatformGroup(UnrealPlatformGroup.Apple))
 			{
 				Lib = "BinkUnreal" + Platform + ".a";
 				PublicDependencyModuleNames.Add("MetalRHI");
