Removing all /Private/ include paths from Media plugins.

diff --git a/Engine/Plugins/Media/AjaMedia/Source/AjaMedia/AjaMedia.Build.cs b/Engine/Plugins/Media/AjaMedia/Source/AjaMedia/AjaMedia.Build.cs
--- a/Engine/Plugins/Media/AjaMedia/Source/AjaMedia/AjaMedia.Build.cs
+++ b/Engine/Plugins/Media/AjaMedia/Source/AjaMedia/AjaMedia.Build.cs
@@ -37,15 +37,6 @@ namespace UnrealBuildTool.Rules
 					"Media",
 				});
 
-			PrivateIncludePaths.AddRange(
-				new string[] {
-					"AjaMedia/Private",
-					"AjaMedia/Private/Aja",
-					"AjaMedia/Private/Assets",
-					"AjaMedia/Private/Player",
-					"AjaMedia/Private/Shared",
-				});
-
 			PublicDependencyModuleNames.AddRange(
 				new string[] {
 					"MediaAssets",
diff --git a/Engine/Plugins/Media/AjaMedia/Source/AjaMedia/Private/Aja/Aja.cpp b/Engine/Plugins/Media/AjaMedia/Source/AjaMedia/Private/Aja/Aja.cpp
--- a/Engine/Plugins/Media/AjaMedia/Source/AjaMedia/Private/Aja/Aja.cpp
+++ b/Engine/Plugins/Media/AjaMedia/Source/AjaMedia/Private/Aja/Aja.cpp
@@ -1,6 +1,6 @@
 // Copyright Epic Games, Inc. All Rights Reserved.
 
-#include "Aja.h"
+#include "Aja/Aja.h"
 #include "AjaMediaPrivate.h"
 
 #include "Misc/FrameRate.h"
diff --git a/Engine/Plugins/Media/AjaMedia/Source/AjaMedia/Private/Aja/AjaDeviceProvider.cpp b/Engine/Plugins/Media/AjaMedia/Source/AjaMedia/Private/Aja/AjaDeviceProvider.cpp
--- a/Engine/Plugins/Media/AjaMedia/Source/AjaMedia/Private/Aja/AjaDeviceProvider.cpp
+++ b/Engine/Plugins/Media/AjaMedia/Source/AjaMedia/Private/Aja/AjaDeviceProvider.cpp
@@ -2,7 +2,7 @@
 
 #include "AjaDeviceProvider.h"
 
-#include "Aja.h"
+#include "Aja/Aja.h"
 #include "AjaMediaPrivate.h"
 
 #include "Async/Async.h"
diff --git a/Engine/Plugins/Media/AjaMedia/Source/AjaMedia/Private/Assets/AjaCustomTimeStep.cpp b/Engine/Plugins/Media/AjaMedia/Source/AjaMedia/Private/Assets/AjaCustomTimeStep.cpp
--- a/Engine/Plugins/Media/AjaMedia/Source/AjaMedia/Private/Assets/AjaCustomTimeStep.cpp
+++ b/Engine/Plugins/Media/AjaMedia/Source/AjaMedia/Private/Assets/AjaCustomTimeStep.cpp
@@ -2,7 +2,7 @@
 
 #include "AjaCustomTimeStep.h"
 #include "AjaMediaPrivate.h"
-#include "AJA.h"
+#include "Aja/Aja.h"
 #include "HAL/IConsoleManager.h"
 #include "HAL/PlatformTime.h"
 #include "Misc/App.h"
diff --git a/Engine/Plugins/Media/AjaMedia/Source/AjaMedia/Private/Assets/AjaMediaSource.cpp b/Engine/Plugins/Media/AjaMedia/Source/AjaMedia/Private/Assets/AjaMediaSource.cpp
--- a/Engine/Plugins/Media/AjaMedia/Source/AjaMedia/Private/Assets/AjaMediaSource.cpp
+++ b/Engine/Plugins/Media/AjaMedia/Source/AjaMedia/Private/Assets/AjaMediaSource.cpp
@@ -2,7 +2,7 @@
 
 #include "AjaMediaSource.h"
 
-#include "Aja.h"
+#include "Aja/Aja.h"
 #include "AjaDeviceProvider.h"
 #include "AjaMediaPrivate.h"
 #include "MediaIOCorePlayerBase.h"
diff --git a/Engine/Plugins/Media/AjaMedia/Source/AjaMedia/Private/Assets/AjaTimecodeProvider.cpp b/Engine/Plugins/Media/AjaMedia/Source/AjaMedia/Private/Assets/AjaTimecodeProvider.cpp
--- a/Engine/Plugins/Media/AjaMedia/Source/AjaMedia/Private/Assets/AjaTimecodeProvider.cpp
+++ b/Engine/Plugins/Media/AjaMedia/Source/AjaMedia/Private/Assets/AjaTimecodeProvider.cpp
@@ -2,7 +2,7 @@
 
 #include "AjaTimecodeProvider.h"
 
-#include "AJA.h"
+#include "Aja/Aja.h"
 #include "AjaMediaPrivate.h"
 #include "Misc/App.h"
 #include "UObject/VirtualProductionObjectVersion.h"
diff --git a/Engine/Plugins/Media/AjaMedia/Source/AjaMedia/Private/Player/AjaMediaPlayer.cpp b/Engine/Plugins/Media/AjaMedia/Source/AjaMedia/Private/Player/AjaMediaPlayer.cpp
--- a/Engine/Plugins/Media/AjaMedia/Source/AjaMedia/Private/Player/AjaMediaPlayer.cpp
+++ b/Engine/Plugins/Media/AjaMedia/Source/AjaMedia/Private/Player/AjaMediaPlayer.cpp
@@ -3,7 +3,7 @@
 #include "AjaMediaPlayer.h"
 
 
-#include "AJA.h"
+#include "Aja/Aja.h"
 #include "AjaMediaAudioSample.h"
 #include "AjaMediaBinarySample.h"
 #include "AjaMediaPrivate.h"
diff --git a/Engine/Plugins/Media/AndroidCamera/Source/AndroidCamera/AndroidCamera.Build.cs b/Engine/Plugins/Media/AndroidCamera/Source/AndroidCamera/AndroidCamera.Build.cs
--- a/Engine/Plugins/Media/AndroidCamera/Source/AndroidCamera/AndroidCamera.Build.cs
+++ b/Engine/Plugins/Media/AndroidCamera/Source/AndroidCamera/AndroidCamera.Build.cs
@@ -29,12 +29,6 @@ namespace UnrealBuildTool.Rules
 					"Media",
 				});
 
-			PrivateIncludePaths.AddRange(
-				new string[] {
-					"AndroidCamera/Private",
-					"AndroidCamera/Private/Player",
-				});
-
 			string PluginPath = Utils.MakePathRelativeTo(ModuleDirectory, Target.RelativeEnginePath);
 			AdditionalPropertiesForReceipt.Add("AndroidPlugin", Path.Combine(PluginPath, "AndroidCamera_UPL.xml"));
 		}
diff --git a/Engine/Plugins/Media/AndroidMedia/Source/AndroidMedia/AndroidMedia.Build.cs b/Engine/Plugins/Media/AndroidMedia/Source/AndroidMedia/AndroidMedia.Build.cs
--- a/Engine/Plugins/Media/AndroidMedia/Source/AndroidMedia/AndroidMedia.Build.cs
+++ b/Engine/Plugins/Media/AndroidMedia/Source/AndroidMedia/AndroidMedia.Build.cs
@@ -24,12 +24,6 @@ namespace UnrealBuildTool.Rules
 				new string[] {
 					"Media",
 				});
-
-			PrivateIncludePaths.AddRange(
-				new string[] {
-					"AndroidMedia/Private",
-					"AndroidMedia/Private/Player",
-				});
 		}
 	}
 }
diff --git a/Engine/Plugins/Media/AvfMedia/Source/AvfMedia/AvfMedia.Build.cs b/Engine/Plugins/Media/AvfMedia/Source/AvfMedia/AvfMedia.Build.cs
--- a/Engine/Plugins/Media/AvfMedia/Source/AvfMedia/AvfMedia.Build.cs
+++ b/Engine/Plugins/Media/AvfMedia/Source/AvfMedia/AvfMedia.Build.cs
@@ -27,13 +27,6 @@ namespace UnrealBuildTool.Rules
 					"Media",
 				});
 
-			PrivateIncludePaths.AddRange(
-				new string[] {
-					"AvfMedia/Private",
-					"AvfMedia/Private/Player",
-					"AvfMedia/Private/Shared",
-				});
-
 			PublicFrameworks.AddRange(
 				new string[] {
 					"CoreMedia",
diff --git a/Engine/Plugins/Media/BinkMedia/Source/BinkMediaPlayerEditor/BinkMediaPlayerEditor.Build.cs b/Engine/Plugins/Media/BinkMedia/Source/BinkMediaPlayerEditor/BinkMediaPlayerEditor.Build.cs
--- a/Engine/Plugins/Media/BinkMedia/Source/BinkMediaPlayerEditor/BinkMediaPlayerEditor.Build.cs
+++ b/Engine/Plugins/Media/BinkMedia/Source/BinkMediaPlayerEditor/BinkMediaPlayerEditor.Build.cs
@@ -11,14 +11,6 @@ public class BinkMediaPlayerEditor : ModuleRules
 	{
 		DynamicallyLoadedModuleNames.AddRange( new string[] { "AssetTools", "MainFrame", "WorkspaceMenuStructure" } );
 
-		PrivateIncludePaths.AddRange(
-			new string[] {
-				"BinkMediaPlayerEditor/Private",
-				"BinkMediaPlayerEditor/Private/Factories",
-				"BinkMediaPlayerEditor/Private/Models",
-			}
-		);
-
 		PrivateDependencyModuleNames.AddRange(
 			new string[] {
 				"ContentBrowser",
diff --git a/Engine/Plugins/Media/BinkMedia/Source/BinkMediaPlayerEditor/Private/BinkMediaPlayerEditorModule.cpp b/Engine/Plugins/Media/BinkMedia/Source/BinkMediaPlayerEditor/Private/BinkMediaPlayerEditorModule.cpp
--- a/Engine/Plugins/Media/BinkMedia/Source/BinkMediaPlayerEditor/Private/BinkMediaPlayerEditorModule.cpp
+++ b/Engine/Plugins/Media/BinkMedia/Source/BinkMediaPlayerEditor/Private/BinkMediaPlayerEditorModule.cpp
@@ -4,13 +4,13 @@
 
 #include "AssetTypeActions_Base.h"
 #include "BinkMediaPlayerCustomization.h"
-#include "BinkMediaPlayerEditorCommands.h"
+#include "Models/BinkMediaPlayerEditorCommands.h"
 #include "ContentBrowserModule.h"
 #include "BinkMediaPlayerEditorToolkit.h"
 #include "Factories/MaterialFactoryNew.h"
 #include "BinkMediaTexture.h"
 #include "Framework/MultiBox/MultiBoxBuilder.h"
-#include "BinkMediaTextureFactoryNew.h"
+#include "Factories/BinkMediaTextureFactoryNew.h"
 #include "IContentBrowserSingleton.h"
 #include "Brushes/SlateImageBrush.h"
 #include "Interfaces/ITextureEditorModule.h"
diff --git a/Engine/Plugins/Media/BinkMedia/Source/BinkMediaPlayerEditor/Private/BinkMediaPlayerEditorToolkit.cpp b/Engine/Plugins/Media/BinkMedia/Source/BinkMediaPlayerEditor/Private/BinkMediaPlayerEditorToolkit.cpp
--- a/Engine/Plugins/Media/BinkMedia/Source/BinkMediaPlayerEditor/Private/BinkMediaPlayerEditorToolkit.cpp
+++ b/Engine/Plugins/Media/BinkMedia/Source/BinkMediaPlayerEditor/Private/BinkMediaPlayerEditorToolkit.cpp
@@ -3,7 +3,7 @@
 
 #include "BinkMediaPlayerEditorToolkit.h"
 
-#include "BinkMediaPlayerEditorCommands.h"
+#include "Models/BinkMediaPlayerEditorCommands.h"
 #include "Framework/MultiBox/MultiBoxBuilder.h"
 #include "Framework/Commands/UICommandList.h"
 #include "Widgets/Docking/SDockTab.h"
diff --git a/Engine/Plugins/Media/BinkMedia/Source/BinkMediaPlayerEditor/Private/SBinkWidgets.cpp b/Engine/Plugins/Media/BinkMedia/Source/BinkMediaPlayerEditor/Private/SBinkWidgets.cpp
--- a/Engine/Plugins/Media/BinkMedia/Source/BinkMediaPlayerEditor/Private/SBinkWidgets.cpp
+++ b/Engine/Plugins/Media/BinkMedia/Source/BinkMediaPlayerEditor/Private/SBinkWidgets.cpp
@@ -3,7 +3,7 @@
 
 #include "SBinkWidgets.h"
 
-#include "BinkMediaPlayerEditorTexture.h"
+#include "Models/BinkMediaPlayerEditorTexture.h"
 #include "Widgets/SViewport.h"
 #include "DetailsViewArgs.h"
 #include "IDetailsView.h"
diff --git a/Engine/Plugins/Media/BlackmagicMedia/Source/BlackmagicMedia/Private/Player/BlackmagicMediaPlayer.cpp b/Engine/Plugins/Media/BlackmagicMedia/Source/BlackmagicMedia/Private/Player/BlackmagicMediaPlayer.cpp
--- a/Engine/Plugins/Media/BlackmagicMedia/Source/BlackmagicMedia/Private/Player/BlackmagicMediaPlayer.cpp
+++ b/Engine/Plugins/Media/BlackmagicMedia/Source/BlackmagicMedia/Private/Player/BlackmagicMediaPlayer.cpp
@@ -1,6 +1,6 @@
 // Copyright Epic Games, Inc. All Rights Reserved.
 
-#include "BlackmagicMediaPlayer.h"
+#include "Player/BlackmagicMediaPlayer.h"
 
 #include "BlackmagicMediaPrivate.h"
 #include "BlackmagicMediaSource.h"
diff --git a/Engine/Plugins/Media/ImgMedia/Source/ImgMedia/ImgMedia.Build.cs b/Engine/Plugins/Media/ImgMedia/Source/ImgMedia/ImgMedia.Build.cs
--- a/Engine/Plugins/Media/ImgMedia/Source/ImgMedia/ImgMedia.Build.cs
+++ b/Engine/Plugins/Media/ImgMedia/Source/ImgMedia/ImgMedia.Build.cs
@@ -34,17 +34,6 @@ namespace UnrealBuildTool.Rules
 					"Media",
 				});
 
-			PrivateIncludePaths.AddRange(
-				new string[] {
-					"ImgMedia/Private",
-					"ImgMedia/Private/Assets",
-					"ImgMedia/Private/GlobalCache",
-					"ImgMedia/Private/Loader",
-					"ImgMedia/Private/Player",
-					"ImgMedia/Private/Readers",
-					"ImgMedia/Private/Scheduler",
-				});
-
 			PublicDependencyModuleNames.AddRange(
 				new string[] {
 					"MediaAssets",
diff --git a/Engine/Plugins/Media/ImgMedia/Source/ImgMedia/Private/GlobalCache/ImgMediaGlobalCache.cpp b/Engine/Plugins/Media/ImgMedia/Source/ImgMedia/Private/GlobalCache/ImgMediaGlobalCache.cpp
--- a/Engine/Plugins/Media/ImgMedia/Source/ImgMedia/Private/GlobalCache/ImgMediaGlobalCache.cpp
+++ b/Engine/Plugins/Media/ImgMedia/Source/ImgMedia/Private/GlobalCache/ImgMediaGlobalCache.cpp
@@ -1,7 +1,7 @@
 // Copyright Epic Games, Inc. All Rights Reserved.
 
 #include "ImgMediaGlobalCache.h"
-#include "IImgMediaReader.h"
+#include "Readers/IImgMediaReader.h"
 #include "ImgMediaPrivate.h"
 #include "ImgMediaSettings.h"
 #include "Modules/ModuleManager.h"
diff --git a/Engine/Plugins/Media/ImgMedia/Source/ImgMedia/Private/ImgMediaModule.cpp b/Engine/Plugins/Media/ImgMedia/Source/ImgMedia/Private/ImgMediaModule.cpp
--- a/Engine/Plugins/Media/ImgMedia/Source/ImgMedia/Private/ImgMediaModule.cpp
+++ b/Engine/Plugins/Media/ImgMedia/Source/ImgMedia/Private/ImgMediaModule.cpp
@@ -7,9 +7,9 @@
 #include "Misc/QueuedThreadPool.h"
 
 #include "ImgMediaGlobalCache.h"
-#include "ImgMediaPlayer.h"
+#include "Player/ImgMediaPlayer.h"
 #include "ImgMediaSceneViewExtension.h"
-#include "ImgMediaScheduler.h"
+#include "Scheduler/ImgMediaScheduler.h"
 #include "ImgMediaSource.h"
 #include "IImgMediaModule.h"
 
diff --git a/Engine/Plugins/Media/ImgMedia/Source/ImgMedia/Private/Loader/ImgMediaLoader.cpp b/Engine/Plugins/Media/ImgMedia/Source/ImgMedia/Private/Loader/ImgMediaLoader.cpp
--- a/Engine/Plugins/Media/ImgMedia/Source/ImgMedia/Private/Loader/ImgMediaLoader.cpp
+++ b/Engine/Plugins/Media/ImgMedia/Source/ImgMedia/Private/Loader/ImgMediaLoader.cpp
@@ -16,16 +16,16 @@
 #include "UObject/Class.h"
 #include "UObject/UObjectGlobals.h"
 
-#include "GenericImgMediaReader.h"
+#include "Readers/GenericImgMediaReader.h"
 #include "IImageWrapperModule.h"
 #include "IImgMediaModule.h"
-#include "IImgMediaReader.h"
+#include "Readers/IImgMediaReader.h"
 #include "ImgMediaLoaderWork.h"
-#include "ImgMediaScheduler.h"
-#include "ImgMediaTextureSample.h"
+#include "Scheduler/ImgMediaScheduler.h"
+#include "Player/ImgMediaTextureSample.h"
 
 #if IMGMEDIA_EXR_SUPPORTED_PLATFORM
-	#include "ExrImgMediaReader.h"
+	#include "Readers/ExrImgMediaReader.h"
 #endif
 
 
diff --git a/Engine/Plugins/Media/ImgMedia/Source/ImgMedia/Private/Loader/ImgMediaLoaderWork.cpp b/Engine/Plugins/Media/ImgMedia/Source/ImgMedia/Private/Loader/ImgMediaLoaderWork.cpp
--- a/Engine/Plugins/Media/ImgMedia/Source/ImgMedia/Private/Loader/ImgMediaLoaderWork.cpp
+++ b/Engine/Plugins/Media/ImgMedia/Source/ImgMedia/Private/Loader/ImgMediaLoaderWork.cpp
@@ -5,7 +5,7 @@
 
 #include "Misc/ScopeLock.h"
 #include "IImgMediaModule.h"
-#include "IImgMediaReader.h"
+#include "Readers/IImgMediaReader.h"
 #include "ImgMediaLoader.h"
 
 
diff --git a/Engine/Plugins/Media/ImgMedia/Source/ImgMedia/Private/Loader/ImgMediaLoaderWork.h b/Engine/Plugins/Media/ImgMedia/Source/ImgMedia/Private/Loader/ImgMediaLoaderWork.h
--- a/Engine/Plugins/Media/ImgMedia/Source/ImgMedia/Private/Loader/ImgMediaLoaderWork.h
+++ b/Engine/Plugins/Media/ImgMedia/Source/ImgMedia/Private/Loader/ImgMediaLoaderWork.h
@@ -3,7 +3,7 @@
 #pragma once
 
 #include "CoreMinimal.h"
-#include "ImgMediaMipMapInfo.h"
+#include "Assets/ImgMediaMipMapInfo.h"
 #include "Misc/IQueuedWork.h"
 
 class FImgMediaLoader;
diff --git a/Engine/Plugins/Media/ImgMedia/Source/ImgMedia/Private/Player/ImgMediaPlayer.cpp b/Engine/Plugins/Media/ImgMedia/Source/ImgMedia/Private/Player/ImgMediaPlayer.cpp
--- a/Engine/Plugins/Media/ImgMedia/Source/ImgMedia/Private/Player/ImgMediaPlayer.cpp
+++ b/Engine/Plugins/Media/ImgMedia/Source/ImgMedia/Private/Player/ImgMediaPlayer.cpp
@@ -12,11 +12,11 @@
 #include "UObject/UObjectGlobals.h"
 #include "UObject/WeakObjectPtr.h"
 
-#include "ImgMediaLoader.h"
-#include "ImgMediaMipMapInfo.h"
-#include "ImgMediaScheduler.h"
+#include "Loader/ImgMediaLoader.h"
+#include "Assets/ImgMediaMipMapInfo.h"
+#include "Scheduler/ImgMediaScheduler.h"
 #include "ImgMediaSettings.h"
-#include "ImgMediaTextureSample.h"
+#include "Player/ImgMediaTextureSample.h"
 
 #include "MediaPlayer.h"
 #include "MediaPlayerFacade.h"
diff --git a/Engine/Plugins/Media/ImgMedia/Source/ImgMedia/Private/Player/ImgMediaTextureSample.h b/Engine/Plugins/Media/ImgMedia/Source/ImgMedia/Private/Player/ImgMediaTextureSample.h
--- a/Engine/Plugins/Media/ImgMedia/Source/ImgMedia/Private/Player/ImgMediaTextureSample.h
+++ b/Engine/Plugins/Media/ImgMedia/Source/ImgMedia/Private/Player/ImgMediaTextureSample.h
@@ -8,7 +8,7 @@
 #include "Misc/Timespan.h"
 #include "Templates/SharedPointer.h"
 
-#include "IImgMediaReader.h"
+#include "Readers/IImgMediaReader.h"
 
 
 /**
diff --git a/Engine/Plugins/Media/ImgMedia/Source/ImgMedia/Private/Readers/ExrImgMediaReader.cpp b/Engine/Plugins/Media/ImgMedia/Source/ImgMedia/Private/Readers/ExrImgMediaReader.cpp
--- a/Engine/Plugins/Media/ImgMedia/Source/ImgMedia/Private/Readers/ExrImgMediaReader.cpp
+++ b/Engine/Plugins/Media/ImgMedia/Source/ImgMedia/Private/Readers/ExrImgMediaReader.cpp
@@ -13,8 +13,8 @@
 #include "UObject/UObjectGlobals.h"
 #include "HardwareInfo.h"
 #include "IImgMediaModule.h"
-#include "ImgMediaLoader.h"
-#include "ImgMediaMipMapInfo.h"
+#include "Loader/ImgMediaLoader.h"
+#include "Assets/ImgMediaMipMapInfo.h"
 
 DECLARE_MEMORY_STAT(TEXT("EXR Reader Pool Memory."), STAT_ExrMediaReaderPoolMem, STATGROUP_ImgMediaPlugin);
 
diff --git a/Engine/Plugins/Media/ImgMedia/Source/ImgMedia/Private/Readers/ExrImgMediaReaderGpu.cpp b/Engine/Plugins/Media/ImgMedia/Source/ImgMedia/Private/Readers/ExrImgMediaReaderGpu.cpp
--- a/Engine/Plugins/Media/ImgMedia/Source/ImgMedia/Private/Readers/ExrImgMediaReaderGpu.cpp
+++ b/Engine/Plugins/Media/ImgMedia/Source/ImgMedia/Private/Readers/ExrImgMediaReaderGpu.cpp
@@ -18,7 +18,7 @@
 #include "Engine/TextureRenderTarget2D.h"
 #include "GlobalShader.h"
 #include "CommonRenderResources.h"
-#include "ImgMediaLoader.h"
+#include "Loader/ImgMediaLoader.h"
 #include "PostProcess/DrawRectangle.h"
 #include "RHIStaticStates.h"
 #include "SceneUtils.h"
diff --git a/Engine/Plugins/Media/ImgMedia/Source/ImgMedia/Private/Readers/GenericImgMediaReader.cpp b/Engine/Plugins/Media/ImgMedia/Source/ImgMedia/Private/Readers/GenericImgMediaReader.cpp
--- a/Engine/Plugins/Media/ImgMedia/Source/ImgMedia/Private/Readers/GenericImgMediaReader.cpp
+++ b/Engine/Plugins/Media/ImgMedia/Source/ImgMedia/Private/Readers/GenericImgMediaReader.cpp
@@ -5,7 +5,7 @@
 
 #include "IImageWrapper.h"
 #include "IImageWrapperModule.h"
-#include "ImgMediaLoader.h"
+#include "Loader/ImgMediaLoader.h"
 #include "Misc/FileHelper.h"
 #include "Misc/Paths.h"
 
diff --git a/Engine/Plugins/Media/ImgMedia/Source/ImgMedia/Private/Readers/IImgMediaReader.h b/Engine/Plugins/Media/ImgMedia/Source/ImgMedia/Private/Readers/IImgMediaReader.h
--- a/Engine/Plugins/Media/ImgMedia/Source/ImgMedia/Private/Readers/IImgMediaReader.h
+++ b/Engine/Plugins/Media/ImgMedia/Source/ImgMedia/Private/Readers/IImgMediaReader.h
@@ -5,7 +5,7 @@
 #include "CoreTypes.h"
 #include "Misc/FrameRate.h"
 #include "IMediaTextureSample.h"
-#include "ImgMediaMipMapInfo.h"
+#include "Assets/ImgMediaMipMapInfo.h"
 #include "Math/IntPoint.h"
 #include "Templates/SharedPointer.h"
 #include "RHI.h"
diff --git a/Engine/Plugins/Media/ImgMedia/Source/ImgMedia/Private/Scheduler/ImgMediaScheduler.cpp b/Engine/Plugins/Media/ImgMedia/Source/ImgMedia/Private/Scheduler/ImgMediaScheduler.cpp
--- a/Engine/Plugins/Media/ImgMedia/Source/ImgMedia/Private/Scheduler/ImgMediaScheduler.cpp
+++ b/Engine/Plugins/Media/ImgMedia/Source/ImgMedia/Private/Scheduler/ImgMediaScheduler.cpp
@@ -8,8 +8,8 @@
 #include "UObject/Class.h"
 #include "UObject/UObjectGlobals.h"
 
-#include "ImgMediaLoader.h"
-#include "ImgMediaSchedulerThread.h"
+#include "Loader/ImgMediaLoader.h"
+#include "Scheduler/ImgMediaSchedulerThread.h"
 #include "ImgMediaSettings.h"
 
 
diff --git a/Engine/Plugins/Media/ImgMedia/Source/ImgMedia/Private/Tests/ImgMediaTests.cpp b/Engine/Plugins/Media/ImgMedia/Source/ImgMedia/Private/Tests/ImgMediaTests.cpp
--- a/Engine/Plugins/Media/ImgMedia/Source/ImgMedia/Private/Tests/ImgMediaTests.cpp
+++ b/Engine/Plugins/Media/ImgMedia/Source/ImgMedia/Private/Tests/ImgMediaTests.cpp
@@ -5,7 +5,7 @@
 #include "Logging/LogMacros.h"
 #include "Misc/AutomationTest.h"
 
-#include "ImgMediaMipMapInfo.h"
+#include "Assets/ImgMediaMipMapInfo.h"
 
 #if WITH_DEV_AUTOMATION_TESTS
 
diff --git a/Engine/Plugins/Media/ImgMedia/Source/ImgMediaEditor/ImgMediaEditor.Build.cs b/Engine/Plugins/Media/ImgMedia/Source/ImgMediaEditor/ImgMediaEditor.Build.cs
--- a/Engine/Plugins/Media/ImgMedia/Source/ImgMediaEditor/ImgMediaEditor.Build.cs
+++ b/Engine/Plugins/Media/ImgMedia/Source/ImgMediaEditor/ImgMediaEditor.Build.cs
@@ -31,9 +31,6 @@ namespace UnrealBuildTool.Rules
 			PrivateIncludePaths.AddRange(
 				new string[] {
 					"ImgMedia/Private",
-					"ImgMediaEditor/Private",
-					"ImgMediaEditor/Private/Customizations",
-					"ImgMediaEditor/Private/Factories",
 				});
 
 			// Are we using the engine?
diff --git a/Engine/Plugins/Media/ImgMedia/Source/ImgMediaEngine/ImgMediaEngine.Build.cs b/Engine/Plugins/Media/ImgMedia/Source/ImgMediaEngine/ImgMediaEngine.Build.cs
--- a/Engine/Plugins/Media/ImgMedia/Source/ImgMediaEngine/ImgMediaEngine.Build.cs
+++ b/Engine/Plugins/Media/ImgMedia/Source/ImgMediaEngine/ImgMediaEngine.Build.cs
@@ -18,12 +18,6 @@ namespace UnrealBuildTool.Rules
 					"ImgMedia",
 				});
 
-			PrivateIncludePaths.AddRange(
-				new string[] {
-					"ImgMediaEngine/Private",
-					"ImgMediaEngine/Private/Unreal",
-				});
-
 			// Are we using the engine?
 			if (Target.bCompileAgainstEngine)
 			{
diff --git a/Engine/Plugins/Media/MediaCompositing/Source/MediaCompositing/MediaCompositing.Build.cs b/Engine/Plugins/Media/MediaCompositing/Source/MediaCompositing/MediaCompositing.Build.cs
--- a/Engine/Plugins/Media/MediaCompositing/Source/MediaCompositing/MediaCompositing.Build.cs
+++ b/Engine/Plugins/Media/MediaCompositing/Source/MediaCompositing/MediaCompositing.Build.cs
@@ -26,12 +26,6 @@ namespace UnrealBuildTool.Rules
 					"RenderCore",
 					"RHI",
 				});
-
-			PrivateIncludePaths.AddRange(
-				new string[] {
-					"MediaCompositing/Private",
-					"MediaCompositing/Private/MovieScene",
-				});
 		}
 	}
 }
diff --git a/Engine/Plugins/Media/MediaCompositing/Source/MediaCompositingEditor/MediaCompositingEditor.Build.cs b/Engine/Plugins/Media/MediaCompositing/Source/MediaCompositingEditor/MediaCompositingEditor.Build.cs
--- a/Engine/Plugins/Media/MediaCompositing/Source/MediaCompositingEditor/MediaCompositingEditor.Build.cs
+++ b/Engine/Plugins/Media/MediaCompositing/Source/MediaCompositingEditor/MediaCompositingEditor.Build.cs
@@ -36,13 +36,6 @@ namespace UnrealBuildTool.Rules
 				{
 					"MovieSceneTools",
 				});
-
-			PrivateIncludePaths.AddRange(
-				new string[] {
-					"MediaCompositingEditor/Private",
-					"MediaCompositingEditor/Private/Sequencer",
-					"MediaCompositingEditor/Private/Shared",
-				});
 		}
 	}
 }
diff --git a/Engine/Plugins/Media/MediaCompositing/Source/MediaCompositingEditor/Private/MediaCompositingEditor.cpp b/Engine/Plugins/Media/MediaCompositing/Source/MediaCompositingEditor/Private/MediaCompositingEditor.cpp
--- a/Engine/Plugins/Media/MediaCompositing/Source/MediaCompositingEditor/Private/MediaCompositingEditor.cpp
+++ b/Engine/Plugins/Media/MediaCompositing/Source/MediaCompositingEditor/Private/MediaCompositingEditor.cpp
@@ -5,10 +5,10 @@
 #include "Modules/ModuleManager.h"
 
 #include "ISequenceRecorder.h"
-#include "MediaCompositingEditorStyle.h"
-#include "MediaSequenceRecorderExtender.h"
+#include "Shared/MediaCompositingEditorStyle.h"
+#include "Sequencer/MediaSequenceRecorderExtender.h"
 #include "Sequencer/MediaTrackEditor.h"
-#include "MediaPlayerPropertyTrackEditor.h"
+#include "Sequencer/MediaPlayerPropertyTrackEditor.h"
 
 
 #define LOCTEXT_NAMESPACE "MediaCompositingEditorModule"
diff --git a/Engine/Plugins/Media/MediaCompositing/Source/MediaCompositingEditor/Private/Shared/MediaCompositingEditorStyle.cpp b/Engine/Plugins/Media/MediaCompositing/Source/MediaCompositingEditor/Private/Shared/MediaCompositingEditorStyle.cpp
--- a/Engine/Plugins/Media/MediaCompositing/Source/MediaCompositingEditor/Private/Shared/MediaCompositingEditorStyle.cpp
+++ b/Engine/Plugins/Media/MediaCompositing/Source/MediaCompositingEditor/Private/Shared/MediaCompositingEditorStyle.cpp
@@ -1,6 +1,6 @@
 // Copyright Epic Games, Inc. All Rights Reserved.
 
-#include "MediaCompositingEditorStyle.h"
+#include "Shared/MediaCompositingEditorStyle.h"
 
 #include "Math/Vector2D.h"
 #include "Styling/SlateStyleRegistry.h"
diff --git a/Engine/Plugins/Media/MediaPlayerEditor/Source/MediaPlayerEditor/MediaPlayerEditor.Build.cs b/Engine/Plugins/Media/MediaPlayerEditor/Source/MediaPlayerEditor/MediaPlayerEditor.Build.cs
--- a/Engine/Plugins/Media/MediaPlayerEditor/Source/MediaPlayerEditor/MediaPlayerEditor.Build.cs
+++ b/Engine/Plugins/Media/MediaPlayerEditor/Source/MediaPlayerEditor/MediaPlayerEditor.Build.cs
@@ -15,18 +15,6 @@ public class MediaPlayerEditor : ModuleRules
 				"WorkspaceMenuStructure",
 			});
 
-		PrivateIncludePaths.AddRange(
-			new string[] {
-				"MediaPlayerEditor/Private",
-				"MediaPlayerEditor/Private/AssetTools",
-				"MediaPlayerEditor/Private/Customizations",
-				"MediaPlayerEditor/Private/Factories",
-				"MediaPlayerEditor/Private/Models",
-				"MediaPlayerEditor/Private/Shared",
-				"MediaPlayerEditor/Private/Toolkits",
-				"MediaPlayerEditor/Private/Widgets",
-			});
-
 		PublicDependencyModuleNames.AddRange(
 			new string[] {
 				"Core",
diff --git a/Engine/Plugins/Media/MediaPlayerEditor/Source/MediaPlayerEditor/Private/Widgets/SMediaPlayerEditorViewer.cpp b/Engine/Plugins/Media/MediaPlayerEditor/Source/MediaPlayerEditor/Private/Widgets/SMediaPlayerEditorViewer.cpp
--- a/Engine/Plugins/Media/MediaPlayerEditor/Source/MediaPlayerEditor/Private/Widgets/SMediaPlayerEditorViewer.cpp
+++ b/Engine/Plugins/Media/MediaPlayerEditor/Source/MediaPlayerEditor/Private/Widgets/SMediaPlayerEditorViewer.cpp
@@ -33,7 +33,7 @@
 #include "Widgets/Notifications/SProgressBar.h"
 #include "Widgets/Text/STextBlock.h"
 
-#include "MediaPlayerEditorSettings.h"
+#include "Shared/MediaPlayerEditorSettings.h"
 #include "SMediaPlayerEditorCache.h"
 #include "SMediaPlayerEditorViewport.h"
 
diff --git a/Engine/Plugins/Media/MediaPlayerEditor/Source/MediaPlayerEditor/Private/Widgets/SMediaPlayerEditorViewport.cpp b/Engine/Plugins/Media/MediaPlayerEditor/Source/MediaPlayerEditor/Private/Widgets/SMediaPlayerEditorViewport.cpp
--- a/Engine/Plugins/Media/MediaPlayerEditor/Source/MediaPlayerEditor/Private/Widgets/SMediaPlayerEditorViewport.cpp
+++ b/Engine/Plugins/Media/MediaPlayerEditor/Source/MediaPlayerEditor/Private/Widgets/SMediaPlayerEditorViewport.cpp
@@ -10,7 +10,7 @@
 #include "Widgets/Layout/SScaleBox.h"
 #include "Widgets/Text/STextBlock.h"
 
-#include "MediaPlayerEditorSettings.h"
+#include "Shared/MediaPlayerEditorSettings.h"
 #include "SMediaPlayerEditorOutput.h"
 #include "SMediaPlayerEditorOverlay.h"
 
diff --git a/Engine/Plugins/Media/MfMedia/Source/MfMedia/MfMedia.Build.cs b/Engine/Plugins/Media/MfMedia/Source/MfMedia/MfMedia.Build.cs
--- a/Engine/Plugins/Media/MfMedia/Source/MfMedia/MfMedia.Build.cs
+++ b/Engine/Plugins/Media/MfMedia/Source/MfMedia/MfMedia.Build.cs
@@ -23,13 +23,6 @@ namespace UnrealBuildTool.Rules
 					"Media",
 				});
 
-			PrivateIncludePaths.AddRange(
-				new string[] {
-					"MfMedia/Private",
-					"MfMedia/Private/Mf",
-					"MfMedia/Private/Player",
-				});
-
 			if (Target.Type != TargetType.Server)
 			{
 				if (Target.Platform == UnrealTargetPlatform.Win64)
diff --git a/Engine/Plugins/Media/MfMedia/Source/MfMedia/Private/MfMediaModule.cpp b/Engine/Plugins/Media/MfMedia/Source/MfMedia/Private/MfMediaModule.cpp
--- a/Engine/Plugins/Media/MfMedia/Source/MfMedia/Private/MfMediaModule.cpp
+++ b/Engine/Plugins/Media/MfMedia/Source/MfMedia/Private/MfMediaModule.cpp
@@ -5,7 +5,7 @@
 #include "Modules/ModuleManager.h"
 
 #if MFMEDIA_SUPPORTED_PLATFORM
-	#include "MfMediaPlayer.h"
+	#include "Player/MfMediaPlayer.h"
 	#include "Templates/SharedPointer.h"
 #endif
 
diff --git a/Engine/Plugins/Media/MfMedia/Source/MfMedia/Private/Player/MfMediaAudioSample.h b/Engine/Plugins/Media/MfMedia/Source/MfMedia/Private/Player/MfMediaAudioSample.h
--- a/Engine/Plugins/Media/MfMedia/Source/MfMedia/Private/Player/MfMediaAudioSample.h
+++ b/Engine/Plugins/Media/MfMedia/Source/MfMedia/Private/Player/MfMediaAudioSample.h
@@ -14,7 +14,7 @@
 #include "Microsoft/COMPointer.h"
 
 #include "MfMediaSample.h"
-#include "MfMediaUtils.h"
+#include "Mf/MfMediaUtils.h"
 
 
 /**
diff --git a/Engine/Plugins/Media/MfMedia/Source/MfMedia/Private/Player/MfMediaPlayer.cpp b/Engine/Plugins/Media/MfMedia/Source/MfMedia/Private/Player/MfMediaPlayer.cpp
--- a/Engine/Plugins/Media/MfMedia/Source/MfMedia/Private/Player/MfMediaPlayer.cpp
+++ b/Engine/Plugins/Media/MfMedia/Source/MfMedia/Private/Player/MfMediaPlayer.cpp
@@ -10,9 +10,9 @@
 #include "MediaSamples.h"
 #include "MediaPlayerOptions.h"
 
-#include "MfMediaSourceReaderCallback.h"
+#include "Mf/MfMediaSourceReaderCallback.h"
 #include "MfMediaTracks.h"
-#include "MfMediaUtils.h"
+#include "Mf/MfMediaUtils.h"
 
 #if PLATFORM_MICROSOFT
 	#include "Microsoft/WindowsHWrapper.h"
diff --git a/Engine/Plugins/Media/MfMedia/Source/MfMedia/Private/Player/MfMediaPlayer.h b/Engine/Plugins/Media/MfMedia/Source/MfMedia/Private/Player/MfMediaPlayer.h
--- a/Engine/Plugins/Media/MfMedia/Source/MfMedia/Private/Player/MfMediaPlayer.h
+++ b/Engine/Plugins/Media/MfMedia/Source/MfMedia/Private/Player/MfMediaPlayer.h
@@ -15,7 +15,7 @@
 #include "Templates/SharedPointer.h"
 #include "Microsoft/COMPointer.h"
 
-#include "IMfMediaSourceReaderSink.h"
+#include "Mf/IMfMediaSourceReaderSink.h"
 
 #if PLATFORM_MICROSOFT
 	#include "Microsoft/WindowsHWrapper.h"
diff --git a/Engine/Plugins/Media/MfMedia/Source/MfMedia/Private/Player/MfMediaSample.h b/Engine/Plugins/Media/MfMedia/Source/MfMedia/Private/Player/MfMediaSample.h
--- a/Engine/Plugins/Media/MfMedia/Source/MfMedia/Private/Player/MfMediaSample.h
+++ b/Engine/Plugins/Media/MfMedia/Source/MfMedia/Private/Player/MfMediaSample.h
@@ -2,7 +2,7 @@
 
 #pragma once
 
-#include "MfMediaUtils.h"
+#include "Mf/MfMediaUtils.h"
 
 
 /**
diff --git a/Engine/Plugins/Media/MfMedia/Source/MfMedia/Private/Player/MfMediaTracks.cpp b/Engine/Plugins/Media/MfMedia/Source/MfMedia/Private/Player/MfMediaTracks.cpp
--- a/Engine/Plugins/Media/MfMedia/Source/MfMedia/Private/Player/MfMediaTracks.cpp
+++ b/Engine/Plugins/Media/MfMedia/Source/MfMedia/Private/Player/MfMediaTracks.cpp
@@ -16,7 +16,7 @@
 
 #include "MfMediaAudioSample.h"
 #include "MfMediaTextureSample.h"
-#include "MfMediaUtils.h"
+#include "Mf/MfMediaUtils.h"
 
 #if PLATFORM_MICROSOFT
 	#include "Microsoft/WindowsHWrapper.h"
diff --git a/Engine/Plugins/Media/WmfMedia/Source/WmfMedia/Private/Player/WmfMediaPlayer.cpp b/Engine/Plugins/Media/WmfMedia/Source/WmfMedia/Private/Player/WmfMediaPlayer.cpp
--- a/Engine/Plugins/Media/WmfMedia/Source/WmfMedia/Private/Player/WmfMediaPlayer.cpp
+++ b/Engine/Plugins/Media/WmfMedia/Source/WmfMedia/Private/Player/WmfMediaPlayer.cpp
@@ -12,10 +12,10 @@
 #include "UObject/Class.h"
 #include "HardwareInfo.h"
 
-#include "WmfMediaSession.h"
+#include "Wmf/WmfMediaSession.h"
 #include "WmfMediaSettings.h"
 #include "WmfMediaTracks.h"
-#include "WmfMediaUtils.h"
+#include "Wmf/WmfMediaUtils.h"
 
 #include "Windows/AllowWindowsPlatformTypes.h"
 
diff --git a/Engine/Plugins/Media/WmfMedia/Source/WmfMedia/Private/Player/WmfMediaTracks.cpp b/Engine/Plugins/Media/WmfMedia/Source/WmfMedia/Private/Player/WmfMediaTracks.cpp
--- a/Engine/Plugins/Media/WmfMedia/Source/WmfMedia/Private/Player/WmfMediaTracks.cpp
+++ b/Engine/Plugins/Media/WmfMedia/Source/WmfMedia/Private/Player/WmfMediaTracks.cpp
@@ -26,12 +26,12 @@
 #include "WmfMediaCodec/WmfMediaDecoder.h"
 #include "WmfMediaHardwareVideoDecodingTextureSample.h"
 #include "WmfMediaOverlaySample.h"
-#include "WmfMediaSampler.h"
+#include "Wmf/WmfMediaSampler.h"
 #include "WmfMediaSettings.h"
-#include "WmfMediaSink.h"
-#include "WmfMediaStreamSink.h"
-#include "WmfMediaTopologyLoader.h"
-#include "WmfMediaUtils.h"
+#include "Wmf/WmfMediaSink.h"
+#include "Wmf/WmfMediaStreamSink.h"
+#include "Wmf/WmfMediaTopologyLoader.h"
+#include "Wmf/WmfMediaUtils.h"
 
 #include "Windows/AllowWindowsPlatformTypes.h"
 
diff --git a/Engine/Plugins/Media/WmfMedia/Source/WmfMedia/Private/Wmf/WmfMediaSession.cpp b/Engine/Plugins/Media/WmfMedia/Source/WmfMedia/Private/Wmf/WmfMediaSession.cpp
--- a/Engine/Plugins/Media/WmfMedia/Source/WmfMedia/Private/Wmf/WmfMediaSession.cpp
+++ b/Engine/Plugins/Media/WmfMedia/Source/WmfMedia/Private/Wmf/WmfMediaSession.cpp
@@ -7,7 +7,7 @@
 #include "IMediaEventSink.h"
 #include "Misc/ScopeLock.h"
 
-#include "WmfMediaTracks.h"
+#include "Player/WmfMediaTracks.h"
 #include "WmfMediaUtils.h"
 
 #include "Windows/AllowWindowsPlatformTypes.h"
diff --git a/Engine/Plugins/Media/WmfMedia/Source/WmfMedia/Private/WmfMediaModule.cpp b/Engine/Plugins/Media/WmfMedia/Source/WmfMedia/Private/WmfMediaModule.cpp
--- a/Engine/Plugins/Media/WmfMedia/Source/WmfMedia/Private/WmfMediaModule.cpp
+++ b/Engine/Plugins/Media/WmfMedia/Source/WmfMedia/Private/WmfMediaModule.cpp
@@ -19,8 +19,8 @@
 	#include "Templates/SharedPointer.h"
 	#include "ShaderCore.h"
 
-	#include "WmfMediaPlayer.h"
-	#include "WmfMediaUtils.h"
+	#include "Player/WmfMediaPlayer.h"
+	#include "Wmf/WmfMediaUtils.h"
 
 	#pragma comment(lib, "mf")
 	#pragma comment(lib, "mfplat")
diff --git a/Engine/Plugins/Media/WmfMedia/Source/WmfMedia/WmfMedia.Build.cs b/Engine/Plugins/Media/WmfMedia/Source/WmfMedia/WmfMedia.Build.cs
--- a/Engine/Plugins/Media/WmfMedia/Source/WmfMedia/WmfMedia.Build.cs
+++ b/Engine/Plugins/Media/WmfMedia/Source/WmfMedia/WmfMedia.Build.cs
@@ -36,13 +36,6 @@ namespace UnrealBuildTool.Rules
 					"Media",
 				});
 
-			PrivateIncludePaths.AddRange(
-				new string[] {
-					"WmfMedia/Private",
-					"WmfMedia/Private/Player",
-					"WmfMedia/Private/Wmf",
-                });
-
 			if (Target.bCompileAgainstEngine)
 			{
 				PrivateDependencyModuleNames.Add("Engine");
