New SDK builds

diff --git a/Engine/Plugins/Media/BinkMedia/Source/BinkMediaPlayerSDK/include/binkplugin.h b/Engine/Plugins/Media/BinkMedia/Source/BinkMediaPlayerSDK/include/binkplugin.h
--- a/Engine/Plugins/Media/BinkMedia/Source/BinkMediaPlayerSDK/include/binkplugin.h
+++ b/Engine/Plugins/Media/BinkMedia/Source/BinkMediaPlayerSDK/include/binkplugin.h
@@ -64,7 +64,7 @@ typedef struct BINKPLUGININITINFO
   // Vulkan only = VkPhysicalDevice
   void * physical_device; 
   
-  // Vulkan only = formats for the sdr and hdr formats
+  // Vulkan/Metal/D3D12 = formats for the sdr and hdr formats
   U32 sdr_and_hdr_render_target_formats[2];
 
   // general CPU allocators (can be zero on most platforms)
diff --git a/Engine/Plugins/Media/BinkMedia/Source/BinkMediaPlayerSDK/include/binktextures.h b/Engine/Plugins/Media/BinkMedia/Source/BinkMediaPlayerSDK/include/binktextures.h
--- a/Engine/Plugins/Media/BinkMedia/Source/BinkMediaPlayerSDK/include/binktextures.h
+++ b/Engine/Plugins/Media/BinkMedia/Source/BinkMediaPlayerSDK/include/binktextures.h
@@ -1,3 +1,7 @@
+// Copyright Epic Games, Inc. All Rights Reserved.
+// Licensed under the terms of a valid Unreal Engine license agreement,
+//   and the separate 'Unreal Engine End User License Agreement for Publishing'.
+
 // Copyright Epic Games Tools, LLC. All Rights Reserved.
 
 /*
diff --git a/Engine/Plugins/Media/BinkMedia/Source/BinkMediaPlayerSDK/include/binktiny.h b/Engine/Plugins/Media/BinkMedia/Source/BinkMediaPlayerSDK/include/binktiny.h
--- a/Engine/Plugins/Media/BinkMedia/Source/BinkMediaPlayerSDK/include/binktiny.h
+++ b/Engine/Plugins/Media/BinkMedia/Source/BinkMediaPlayerSDK/include/binktiny.h
@@ -2,6 +2,8 @@
 // Licensed under the terms of a valid Unreal Engine license agreement,
 //   and the separate 'Unreal Engine End User License Agreement for Publishing'.
 
+// Copyright Epic Games Tools, LLC. All Rights Reserved.
+
 #ifndef __BINKH__
 #define __BINKH__
 
diff --git a/Engine/Plugins/Media/BinkMedia/Source/BinkMediaPlayerSDK/include/egttypes.h b/Engine/Plugins/Media/BinkMedia/Source/BinkMediaPlayerSDK/include/egttypes.h
--- a/Engine/Plugins/Media/BinkMedia/Source/BinkMediaPlayerSDK/include/egttypes.h
+++ b/Engine/Plugins/Media/BinkMedia/Source/BinkMediaPlayerSDK/include/egttypes.h
@@ -3,7 +3,7 @@
 //   and the separate 'Unreal Engine End User License Agreement for Publishing'.
 
 //===============================================
-// (C) Copyright 1994-2022 Epic Games Tools, LLC  
+// (C) Copyright 1994-2023 Epic Games Tools, LLC  
 //===============================================
 
 #ifndef __EGTTYPESH__
@@ -13,7 +13,7 @@
 #define __RADTYPESH__
 #define __RADRR_COREH__ // block old rr_core
 
-#define RADCOPYRIGHT "Copyright (C) 1994-2022, Epic Games Tools, LLC"
+#define RADCOPYRIGHT "Copyright (C) 1994-2023, Epic Games Tools, LLC"
 
 #if !defined(__RADRES__) // don't include anything for resource compiles
 
@@ -36,11 +36,13 @@
 //  __RADEMSCRIPTEN__ means Emscripten
 
 //  __RADARM__ means arm
+//  __RADARM64__ means arm64  (+ __RADARM__ will also be set)
 //  __RADPPC__ means powerpc
 //  __RADX86__ means x86 or x64
 //  __RADX64__ means x64  (+ __RADX86__ will also be set)
 //  __RADNEON__ means you can use NEON intrinsics on ARM
 //  __RADSSE2__ means you can use SSE2 intrinsics on X86
+//  __RADSSE4__ means you can use SSE4.1 intrinsics on X86
 
 // __RADNOVARARGMACROS__ means #defines can't use ...
 
@@ -227,7 +229,7 @@
   #define __RADDETECTEDPROC__ __RADPPC__
   #define __RADBIGENDIAN__
 #endif
-#if defined( __aarch64__ ) || defined( __arm64__ )
+#if defined( __aarch64__ ) || defined( __arm64__ ) || defined( _M_ARM64 )
   #define __RADARM__ 1
   #define __RADARM64__ 6
   #define __RADDETECTEDPROC__ __RADARM64__
@@ -243,7 +245,7 @@
   #error "egttypes.h did not detect your processor type."
 #endif
 
-#if defined(__ppc64__) || defined(__aarch64__) || defined(_M_X64) || defined(__x86_64__) || defined(__x86_64)
+#if defined(__ppc64__) || defined(__aarch64__) || defined(_M_X64) || defined(__x86_64__) || defined(__x86_64) || defined( __arm64__ ) || defined( _M_ARM64 )
   #define __RAD64__
   #define __RAD64REGS__  // need to set this for platforms that aren't 64-bit, but have 64-bit regs (old consoles)
 #endif
