Bink: Integrate debug log printing from Git pull request
#11062

diff --git a/Engine/Plugins/Media/BinkMedia/Source/BinkMediaPlayer/Private/Assets/BinkMediaPlayer.cpp b/Engine/Plugins/Media/BinkMedia/Source/BinkMediaPlayer/Private/Assets/BinkMediaPlayer.cpp
--- a/Engine/Plugins/Media/BinkMedia/Source/BinkMediaPlayer/Private/Assets/BinkMediaPlayer.cpp
+++ b/Engine/Plugins/Media/BinkMedia/Source/BinkMediaPlayer/Private/Assets/BinkMediaPlayer.cpp
@@ -284,11 +284,19 @@ bool UBinkMediaPlayer::Open(const FString& Url)
 {
 	if (!BinkInitialize())
 	{
+		UE_LOG(LogBinkMoviePlayer, Error, TEXT("UBinkMediaPlayer::Open: failed to initialize bink!"));
 		return false;
 	}
 
-	if (Url.IsEmpty() || IsPlaying()) 
+	if (Url.IsEmpty())
 	{
+		UE_LOG(LogBinkMoviePlayer, Error, TEXT("UBinkMediaPlayer::Open: Failed! Url is empty."));
+		return false;
+	}
+
+	if (IsPlaying()) 
+	{
+		UE_LOG(LogBinkMoviePlayer, Error, TEXT("UBinkMediaPlayer::Open: Failed! Already playing."));
 		return false;
 	}
 
@@ -317,6 +325,7 @@ bool UBinkMediaPlayer::Open(const FString& Url)
 
 	// Don't bother trying to play it if we can't find it.
 	if (!IAndroidPlatformFile::GetPlatformPhysical().FileExists(*MoviePath)) {
+		UE_LOG(LogBinkMoviePlayer, Error, TEXT("UBinkMediaPlayer::Open: Failed! File doesn't exist. URL :%s File :%s"), *Url, *MoviePath);
 		return false;
 	}
 
@@ -356,6 +365,7 @@ bool UBinkMediaPlayer::Open(const FString& Url)
 	if(!bnk) 
 	{
 		//MediaEvent.Broadcast(EMediaEvent::MediaOpenFailed);
+		UE_LOG(LogBinkMoviePlayer, Error, TEXT("UBinkMediaPlayer::Open: Failed! BinkPluginOpen failed. Invalid bnk."));
 		return false;
 	}
 
