Bink for new platform SDKs

diff --git a/Engine/Plugins/Media/BinkMedia/Source/BinkMediaPlayerSDK/include/egttypes.h b/Engine/Plugins/Media/BinkMedia/Source/BinkMediaPlayerSDK/include/egttypes.h
--- a/Engine/Plugins/Media/BinkMedia/Source/BinkMediaPlayerSDK/include/egttypes.h
+++ b/Engine/Plugins/Media/BinkMedia/Source/BinkMediaPlayerSDK/include/egttypes.h
@@ -3,7 +3,7 @@
 //   and the separate 'Unreal Engine End User License Agreement for Publishing'.
 
 //===============================================
-// (C) Copyright 1994-2023 Epic Games Tools, LLC  
+// Copyright Epic Games Tools, LLC. All Rights Reserved.
 //===============================================
 
 #ifndef __EGTTYPESH__
@@ -13,7 +13,7 @@
 #define __RADTYPESH__
 #define __RADRR_COREH__ // block old rr_core
 
-#define RADCOPYRIGHT "Copyright (C) 1994-2023, Epic Games Tools, LLC"
+#define RADCOPYRIGHT "Copyright (C) 1994-2024, Epic Games Tools, LLC"
 
 #if !defined(__RADRES__) // don't include anything for resource compiles
 
