Added GetTexture interface to BinkMovieStreamer

diff --git a/Engine/Plugins/Media/BinkMedia/Source/BinkMediaPlayer/Private/BinkMovieStreamer.h b/Engine/Plugins/Media/BinkMedia/Source/BinkMediaPlayer/Private/BinkMovieStreamer.h
--- a/Engine/Plugins/Media/BinkMedia/Source/BinkMediaPlayer/Private/BinkMovieStreamer.h
+++ b/Engine/Plugins/Media/BinkMedia/Source/BinkMediaPlayer/Private/BinkMovieStreamer.h
@@ -16,13 +16,19 @@ public:
 	virtual ~FBinkMovieStreamer();
 
 	virtual bool Init(const TArray<FString>& MoviePaths, TEnumAsByte<EMoviePlaybackType> inPlaybackType) override;
-	virtual FString GetMovieName() override;
-	virtual bool IsLastMovieInPlaylist() override;
+	//virtual void PreviousViewportInterface(const TSharedPtr<ISlateViewport>& PreviousViewportInterface) override;
 	virtual void ForceCompletion() override;
 	virtual bool Tick(float DeltaTime) override;
 	virtual TSharedPtr<class ISlateViewport> GetViewportInterface() override { return MovieViewport; }
 	virtual float GetAspectRatio() const override { return (float)MovieViewport->GetSize().X / (float)MovieViewport->GetSize().Y; }
+	virtual FString GetMovieName() override;
+	virtual bool IsLastMovieInPlaylist() override;
 	virtual void Cleanup() override;
+	virtual FTextureRHIRef GetTexture() override { return Texture->IsValid() ? Texture->GetRHIRef() : nullptr; }
+
+	//virtual void TickPreEngine() override;
+	//virtual void TickPostEngine() override;
+	//virtual void TickPostRender() override;
 
 	FOnCurrentMovieClipFinished OnCurrentMovieClipFinishedDelegate;
 	virtual FOnCurrentMovieClipFinished& OnCurrentMovieClipFinished() override { return OnCurrentMovieClipFinishedDelegate; }
