Fixed crash on exit where Bink movie player would be destructed at the incorrect time.

diff --git a/Engine/Plugins/Media/BinkMedia/Source/BinkMediaPlayer/Private/BinkMediaPlayerModule.cpp b/Engine/Plugins/Media/BinkMedia/Source/BinkMediaPlayer/Private/BinkMediaPlayerModule.cpp
--- a/Engine/Plugins/Media/BinkMedia/Source/BinkMediaPlayer/Private/BinkMediaPlayerModule.cpp
+++ b/Engine/Plugins/Media/BinkMedia/Source/BinkMediaPlayer/Private/BinkMediaPlayerModule.cpp
@@ -227,6 +227,7 @@ struct FBinkMediaPlayerModule : IModuleInterface, FTickableGameObject
 				GEngine->GameViewport->OnDrawn().Remove(overlayHook);
 			}
 		}
+		MovieStreamer.Reset();
 #if BINKPLUGIN_UE4_EDITOR
 		FEditorDelegates::BeginPIE.RemoveAll(this);
 		FEditorDelegates::EndPIE.RemoveAll(this);
