Bink: Fixed importing assets already in Content/Movies directory.

diff --git a/Engine/Plugins/Media/BinkMedia/Source/BinkMediaPlayerEditor/Private/Factories/BinkMediaPlayerFactory.cpp b/Engine/Plugins/Media/BinkMedia/Source/BinkMediaPlayerEditor/Private/Factories/BinkMediaPlayerFactory.cpp
--- a/Engine/Plugins/Media/BinkMedia/Source/BinkMediaPlayerEditor/Private/Factories/BinkMediaPlayerFactory.cpp
+++ b/Engine/Plugins/Media/BinkMedia/Source/BinkMediaPlayerEditor/Private/Factories/BinkMediaPlayerFactory.cpp
@@ -3,7 +3,9 @@
 
 #include "BinkMediaPlayerFactory.h"
 
+#include "Misc/Paths.h"
 #include "BinkMediaPlayer.h"
+#include "binkplugin_ue4.h"
 
 UBinkMediaPlayerFactory::UBinkMediaPlayerFactory( const FObjectInitializer& ObjectInitializer )
 	: Super(ObjectInitializer)
@@ -19,6 +21,23 @@ UBinkMediaPlayerFactory::UBinkMediaPlayerFactory( const FObjectInitializer& Obje
 UObject* UBinkMediaPlayerFactory::FactoryCreateBinary( UClass* Class, UObject* InParent, FName Name, EObjectFlags Flags, UObject* Context, const TCHAR* Type, const uint8*& Buffer, const uint8* BufferEnd, FFeedbackContext* Warn ) 
 {
 	UBinkMediaPlayer* MediaPlayer = NewObject<UBinkMediaPlayer>(InParent, Class, Name, Flags);
-	MediaPlayer->OpenUrl(CurrentFilename);
+
+	// This workaround is based on FBinkMediaPlayerCustomization::HandleUrlPickerPathPicked.
+	if (CurrentFilename.IsEmpty() || CurrentFilename.StartsWith(TEXT("./")) || CurrentFilename.Contains(TEXT("://")))
+	{
+		MediaPlayer->OpenUrl(CurrentFilename);
+	}
+	else
+	{
+		FString FullUrl = FPaths::ConvertRelativePathToFull(CurrentFilename);
+		const FString FullGameContentDir = FPaths::ConvertRelativePathToFull(BINKCONTENTPATH);
+		if (FullUrl.StartsWith(FullGameContentDir))
+		{
+			FPaths::MakePathRelativeTo(FullUrl, *FullGameContentDir);
+			FullUrl = FString(TEXT("./")) + FullUrl;
+		}
+		MediaPlayer->OpenUrl(FullUrl);
+	}
+
 	return MediaPlayer;
 }
