Bink: Show filename in picker box
#11780

diff --git a/Engine/Plugins/Media/BinkMedia/Source/BinkMediaPlayerEditor/Private/BinkMediaPlayerCustomization.cpp b/Engine/Plugins/Media/BinkMedia/Source/BinkMediaPlayerEditor/Private/BinkMediaPlayerCustomization.cpp
--- a/Engine/Plugins/Media/BinkMedia/Source/BinkMediaPlayerEditor/Private/BinkMediaPlayerCustomization.cpp
+++ b/Engine/Plugins/Media/BinkMedia/Source/BinkMediaPlayerEditor/Private/BinkMediaPlayerCustomization.cpp
@@ -14,6 +14,7 @@
 #include "Misc/Paths.h"
 #include "Widgets/Text/STextBlock.h"
 #include "Widgets/Input/SButton.h"
+#include "Widgets/Input/SEditableText.h"
 #include "Widgets/SWindow.h"
 #include "binkplugin_ue4.h"
 
@@ -133,8 +134,37 @@ void FBinkMediaPlayerCustomization::CustomizeDetails( IDetailLayoutBuilder& Deta
 				.MaxDesiredWidth(0.0f)
 				.MinDesiredWidth(125.0f)
 				[
-					SNew(SBinkFilePathPicker)
-						.OnPathPicked(this, &FBinkMediaPlayerCustomization::HandleUrlPickerPathPicked)
+					SNew(SHorizontalBox)
+						+ SHorizontalBox::Slot()
+						.AutoWidth()
+						.VAlign(VAlign_Center)
+						[
+							SNew(SEditableText)
+								.Font(IDetailLayoutBuilder::GetDetailFontBold())
+								.Text_Lambda([=, this]()
+									{
+										FText UrlValueText;
+										UrlProperty->GetValueAsDisplayText(UrlValueText);
+										return UrlValueText;
+									})
+								.OnTextCommitted_Lambda([this](const FText& Val, ETextCommit::Type TextCommitType)
+									{
+										if (TextCommitType != ETextCommit::OnEnter)
+											return;
+
+										FString NewPathStr = Val.ToString();
+										HandleUrlPickerPathPicked(NewPathStr);
+									})
+						]
+						+ SHorizontalBox::Slot()
+						.FillWidth(1.0f)
+						.HAlign(HAlign_Right)
+						.VAlign(VAlign_Center)
+						.Padding(4.0f, 0.0f, 4.0f, 0.0f)
+						[
+							SNew(SBinkFilePathPicker)
+								.OnPathPicked(this, &FBinkMediaPlayerCustomization::HandleUrlPickerPathPicked)
+						]
 				];
 		}
 	}
