Don't add Movies to runtime dependencies in Bink SDK build cs file.

diff --git a/Engine/Plugins/Media/BinkMedia/Source/BinkMediaPlayerSDK/BinkMediaPlayerSDK.Build.cs b/Engine/Plugins/Media/BinkMedia/Source/BinkMediaPlayerSDK/BinkMediaPlayerSDK.Build.cs
--- a/Engine/Plugins/Media/BinkMedia/Source/BinkMediaPlayerSDK/BinkMediaPlayerSDK.Build.cs
+++ b/Engine/Plugins/Media/BinkMedia/Source/BinkMediaPlayerSDK/BinkMediaPlayerSDK.Build.cs
@@ -50,7 +50,7 @@ public class BinkMediaPlayerSDK : ModuleRules
 		PublicDefinitions.Add("BUILDING_FOR_UNREAL_ONLY=1");
 		PublicDefinitions.Add("__RADNOEXPORTS__=1");
 		PublicDefinitions.Add("__RADINSTATICLIB__=1");
-		RuntimeDependencies.Add("$(ProjectDir)/Content/Movies/..."); // For chunked streaming
+		//RuntimeDependencies.Add("$(ProjectDir)/Content/Movies/..."); // For chunked streaming
 
 		string Lib = LibName;
 		string Platform = Target.Platform.ToString();
