Bink: HDR settings when rendering overlays

diff --git a/Engine/Plugins/Media/BinkMedia/Source/BinkMediaPlayer/Private/Assets/BinkMediaPlayer.cpp b/Engine/Plugins/Media/BinkMedia/Source/BinkMediaPlayer/Private/Assets/BinkMediaPlayer.cpp
--- a/Engine/Plugins/Media/BinkMedia/Source/BinkMediaPlayer/Private/Assets/BinkMediaPlayer.cpp
+++ b/Engine/Plugins/Media/BinkMedia/Source/BinkMediaPlayer/Private/Assets/BinkMediaPlayer.cpp
@@ -10,6 +10,7 @@
 #include "Engine/Texture.h"
 #include "Misc/Paths.h"
 #include "RenderingThread.h"
+#include "HDRHelper.h"
 
 
 #include "HAL/PlatformFileManager.h"
@@ -570,6 +571,46 @@ void UBinkMediaPlayer::Tick(float DeltaTime)
 
 			ENQUEUE_RENDER_COMMAND(BinkScheduleOverlay)([bnk=bnk,ulx,uly,lrx,lry](FRHICommandListImmediate& RHICmdList) 
 			{ 
+				static const auto CVarHDROutputEnabled = IConsoleManager::Get().FindTConsoleVariableDataInt(TEXT("r.HDR.EnableHDROutput"));
+				static const auto CVarDisplayOutputDevice = IConsoleManager::Get().FindTConsoleVariableDataInt(TEXT("r.HDR.Display.OutputDevice"));
+				if (GRHISupportsHDROutput && CVarHDROutputEnabled->GetValueOnRenderThread() != 0)
+				{
+					EDisplayOutputFormat outDev = static_cast<EDisplayOutputFormat>(CVarDisplayOutputDevice->GetValueOnRenderThread());
+					float DisplayMaxLuminance = HDRGetDisplayMaximumLuminance();
+					switch (outDev)
+					{
+					// LDR
+					case EDisplayOutputFormat::SDR_sRGB:
+					case EDisplayOutputFormat::SDR_Rec709:
+					case EDisplayOutputFormat::SDR_ExplicitGammaMapping:
+						BinkPluginSetHdrSettings(bnk, 1, 1.0f, 80);
+						break;
+					// 1k nits
+					case EDisplayOutputFormat::HDR_ACES_1000nit_ST2084:
+						BinkPluginSetHdrSettings(bnk, 2, 1.0f, DisplayMaxLuminance);
+						break;
+					case EDisplayOutputFormat::HDR_ACES_1000nit_ScRGB:
+						BinkPluginSetHdrSettings(bnk, 1, 1.0f, DisplayMaxLuminance);
+						break;
+					// 2k nits
+					case EDisplayOutputFormat::HDR_ACES_2000nit_ST2084:
+						BinkPluginSetHdrSettings(bnk, 2, 1.0f, DisplayMaxLuminance);
+						break;
+					case EDisplayOutputFormat::HDR_ACES_2000nit_ScRGB:
+						BinkPluginSetHdrSettings(bnk, 1, 1.0f, DisplayMaxLuminance);
+						break;
+					// no tonemap
+					default:
+						BinkPluginSetHdrSettings(bnk, 0, 1.0f, 1000);
+						break;
+					}
+					BinkPluginSetRenderTargetFormat(bnk, 1);
+				}
+				else
+				{
+					BinkPluginSetHdrSettings(bnk, 1, 1.0f, 80);
+					BinkPluginSetRenderTargetFormat(bnk, 0);
+				}
 				BinkPluginSetDrawFlags(bnk, 0);
 				BinkPluginScheduleOverlay(bnk, ulx, uly, lrx, lry, 0);
 			});
