Bink: Fixed animating Rect-lights

diff --git a/Engine/Plugins/Media/BinkMedia/Source/BinkMediaPlayer/BinkMediaPlayer.Build.cs b/Engine/Plugins/Media/BinkMedia/Source/BinkMediaPlayer/BinkMediaPlayer.Build.cs
--- a/Engine/Plugins/Media/BinkMedia/Source/BinkMediaPlayer/BinkMediaPlayer.Build.cs
+++ b/Engine/Plugins/Media/BinkMedia/Source/BinkMediaPlayer/BinkMediaPlayer.Build.cs
@@ -15,12 +15,15 @@ public class BinkMediaPlayer : ModuleRules
 		PublicDependencyModuleNames.Add("CoreUObject");
 		PublicDependencyModuleNames.Add("Engine");
 		PublicDependencyModuleNames.Add("InputCore");
-		PublicDependencyModuleNames.Add("RenderCore");
-		PublicDependencyModuleNames.Add("RHI");
 		PublicDependencyModuleNames.Add("MoviePlayer");
 		//PublicDependencyModuleNames.Add("MediaAssets");
 		PublicDependencyModuleNames.Add("Projects");
 		//PublicDependencyModuleNames.Add("PixelStreaming");
+
+		PrivateDependencyModuleNames.Add("RenderCore");
+		PrivateDependencyModuleNames.Add("RHI");
+		//PrivateDependencyModuleNames.Add("ColorManagement");
+		PrivateDependencyModuleNames.Add("Renderer");
 		
 		AddEngineThirdPartyPrivateStaticDependencies(Target, "BinkMediaPlayerSDK");
 		
diff --git a/Engine/Plugins/Media/BinkMedia/Source/BinkMediaPlayer/Private/Assets/BinkMediaPlayer.cpp b/Engine/Plugins/Media/BinkMedia/Source/BinkMediaPlayer/Private/Assets/BinkMediaPlayer.cpp
--- a/Engine/Plugins/Media/BinkMedia/Source/BinkMediaPlayer/Private/Assets/BinkMediaPlayer.cpp
+++ b/Engine/Plugins/Media/BinkMedia/Source/BinkMediaPlayer/Private/Assets/BinkMediaPlayer.cpp
@@ -19,6 +19,7 @@
 #include "UObject/UObjectIterator.h"
 #include "binkplugin_ue4.h"
 #include "egttypes.h"
+#include "RectLightTexture.h"
 
 //#include "Media/Public/IMediaEventSink.h" //For EMediaEvent
 
@@ -618,7 +619,7 @@ void UBinkMediaPlayer::Tick(float DeltaTime)
 	}
 }
 
-void UBinkMediaPlayer::UpdateTexture(FRHICommandListImmediate &RHICmdList, FTextureRHIRef ref, void *nativePtr, int width, int height, bool isEditor, bool tonemap, int output_nits, float alpha, bool srgb_decode, bool is_hdr) 
+void UBinkMediaPlayer::UpdateTexture(const UTexture *tex, FRHICommandListImmediate &RHICmdList, FTextureRHIRef ref, void *nativePtr, int width, int height, bool isEditor, bool tonemap, int output_nits, float alpha, bool srgb_decode, bool is_hdr) 
 {
 	check(IsInRenderingThread());
 
@@ -630,6 +631,13 @@ void UBinkMediaPlayer::UpdateTexture(FRHICommandListImmediate &RHICmdList, FText
 		BinkPluginSetRenderTargetFormat(bnk, is_hdr ? 1 : 0);
 		BinkPluginScheduleToTexture(bnk, BinkDestinationUpperLeft.X, BinkDestinationUpperLeft.Y, BinkDestinationLowerRight.X, BinkDestinationLowerRight.Y, 0, ref.GetReference(), width, height);
 		BinkActiveTextureRefs.Push(ref);
+
+		// Update rect-lights with this texture
+		if (tex)
+		{
+			// Lock/Enqueue rect atlas refresh if that texture is used by a rect. light
+			RectLightAtlas::FAtlasTextureInvalidationScope InvalidationScope(tex);
+		}
 	}
 }
 
@@ -656,6 +664,7 @@ void UBinkMediaPlayer::Draw(UTexture *texture, bool tonemap, int out_nits, float
 	struct parms_t 
 	{
 		UBinkMediaPlayer *player;
+		UTexture* tex;
 		FTextureRHIRef ref;
 		void *native;
 		int width, height;
@@ -664,10 +673,10 @@ void UBinkMediaPlayer::Draw(UTexture *texture, bool tonemap, int out_nits, float
 		bool hdr;
 		int out_nits;
 		float alpha;
-	} parms = { this, ref, native, width, height, tonemap, srgb_decode, hdr, out_nits, alpha };
+	} parms = { this, texture, ref, native, width, height, tonemap, srgb_decode, hdr, out_nits, alpha };
 	ENQUEUE_RENDER_COMMAND(BinkMediaPlayer_Draw)([parms](FRHICommandListImmediate& RHICmdList) 
 	{ 
-		parms.player->UpdateTexture(RHICmdList, parms.ref, parms.native, parms.width, parms.height, false, parms.tonemap, parms.out_nits, parms.alpha, parms.srgb_decode, parms.hdr);
+		parms.player->UpdateTexture(parms.tex, RHICmdList, parms.ref, parms.native, parms.width, parms.height, false, parms.tonemap, parms.out_nits, parms.alpha, parms.srgb_decode, parms.hdr);
 	});
 }
 
diff --git a/Engine/Plugins/Media/BinkMedia/Source/BinkMediaPlayer/Private/Assets/BinkMediaTextureResource.cpp b/Engine/Plugins/Media/BinkMedia/Source/BinkMediaPlayer/Private/Assets/BinkMediaTextureResource.cpp
--- a/Engine/Plugins/Media/BinkMedia/Source/BinkMediaPlayer/Private/Assets/BinkMediaTextureResource.cpp
+++ b/Engine/Plugins/Media/BinkMedia/Source/BinkMediaPlayer/Private/Assets/BinkMediaTextureResource.cpp
@@ -101,7 +101,7 @@ void FBinkMediaTextureResource::UpdateDeferredResource(FRHICommandListImmediate&
 	uint32 width = tex->GetSizeX();
 	uint32 height = tex->GetSizeY();
 	bool is_hdr = PixelFormat != PF_B8G8R8A8;
-	Player->UpdateTexture(RHICmdList, tex, tex->GetNativeResource(), width, height, false, Owner->Tonemap, Owner->OutputNits, Owner->Alpha, Owner->DecodeSRGB, is_hdr);
+	Player->UpdateTexture(Owner, RHICmdList, tex, tex->GetNativeResource(), width, height, false, Owner->Tonemap, Owner->OutputNits, Owner->Alpha, Owner->DecodeSRGB, is_hdr);
 }
 
 void FBinkMediaTextureResource::Clear() 
diff --git a/Engine/Plugins/Media/BinkMedia/Source/BinkMediaPlayer/Public/BinkMediaPlayer.h b/Engine/Plugins/Media/BinkMedia/Source/BinkMediaPlayer/Public/BinkMediaPlayer.h
--- a/Engine/Plugins/Media/BinkMedia/Source/BinkMediaPlayer/Public/BinkMediaPlayer.h
+++ b/Engine/Plugins/Media/BinkMedia/Source/BinkMediaPlayer/Public/BinkMediaPlayer.h
@@ -462,7 +462,7 @@ public:
 	bool Open(const FString& Url);
 	void Close();
 	bool IsReady() const { return bnk != NULL; }
-	void UpdateTexture(FRHICommandListImmediate &RHICmdList, FTextureRHIRef ref, void *nativePtr, int width, int height, bool isEditor, bool tonemap, int output_nits, float alpha, bool srgb_decode, bool is_hdr);
+	void UpdateTexture(const UTexture* tex, FRHICommandListImmediate &RHICmdList, FTextureRHIRef ref, void *nativePtr, int width, int height, bool isEditor, bool tonemap, int output_nits, float alpha, bool srgb_decode, bool is_hdr);
 
 	FIntPoint GetDimensions() const;
 	float GetFrameRate() const;
diff --git a/Engine/Plugins/Media/BinkMedia/Source/BinkMediaPlayerEditor/Private/Models/BinkMediaPlayerEditorTexture.h b/Engine/Plugins/Media/BinkMedia/Source/BinkMediaPlayerEditor/Private/Models/BinkMediaPlayerEditorTexture.h
--- a/Engine/Plugins/Media/BinkMedia/Source/BinkMediaPlayerEditor/Private/Models/BinkMediaPlayerEditorTexture.h
+++ b/Engine/Plugins/Media/BinkMedia/Source/BinkMediaPlayerEditor/Private/Models/BinkMediaPlayerEditorTexture.h
@@ -34,7 +34,7 @@ struct FBinkMediaPlayerEditorTexture : public FTickableGameObject
 				return;
 			}
 
-			MediaPlayer->UpdateTexture(RHICmdList, tex, tex->GetNativeResource(), tex->GetSizeX(), tex->GetSizeY(), true, true, 80, 1, false, false);
+			MediaPlayer->UpdateTexture(NULL, RHICmdList, tex, tex->GetNativeResource(), tex->GetSizeX(), tex->GetSizeY(), true, true, 80, 1, false, false);
 		});
 	}
 	
