Fixed compile errors when using vfsoverlay. Caused by ../../.. paths

diff --git a/Engine/Plugins/Media/AndroidCamera/Source/AndroidCameraFactory/Private/AndroidCameraFactoryModule.cpp b/Engine/Plugins/Media/AndroidCamera/Source/AndroidCameraFactory/Private/AndroidCameraFactoryModule.cpp
--- a/Engine/Plugins/Media/AndroidCamera/Source/AndroidCameraFactory/Private/AndroidCameraFactoryModule.cpp
+++ b/Engine/Plugins/Media/AndroidCamera/Source/AndroidCameraFactory/Private/AndroidCameraFactoryModule.cpp
@@ -2,6 +2,7 @@
 
 #include "AndroidCameraFactoryPrivate.h"
 
+#include "IAndroidCameraModule.h"
 #include "Internationalization/Internationalization.h"
 #include "Misc/Guid.h"
 #include "Misc/Paths.h"
@@ -9,9 +10,6 @@
 #include "IMediaModule.h"
 #include "IMediaPlayerFactory.h"
 
-#include "../../AndroidCamera/Public/IAndroidCameraModule.h"
-
-
 DEFINE_LOG_CATEGORY(LogAndroidCameraFactory);
 
 #define LOCTEXT_NAMESPACE "FAndroidCameraFactoryModule"
diff --git a/Engine/Plugins/Media/BinkMedia/Source/BinkMediaPlayer/BinkMediaPlayer.Build.cs b/Engine/Plugins/Media/BinkMedia/Source/BinkMediaPlayer/BinkMediaPlayer.Build.cs
--- a/Engine/Plugins/Media/BinkMedia/Source/BinkMediaPlayer/BinkMediaPlayer.Build.cs
+++ b/Engine/Plugins/Media/BinkMedia/Source/BinkMediaPlayer/BinkMediaPlayer.Build.cs
@@ -16,13 +16,10 @@ public class BinkMediaPlayer : ModuleRules
 		PublicDependencyModuleNames.Add("Engine");
 		PublicDependencyModuleNames.Add("InputCore");
 		PublicDependencyModuleNames.Add("MoviePlayer");
-		//PublicDependencyModuleNames.Add("MediaAssets");
 		PublicDependencyModuleNames.Add("Projects");
-		//PublicDependencyModuleNames.Add("PixelStreaming");
 
 		PrivateDependencyModuleNames.Add("RenderCore");
 		PrivateDependencyModuleNames.Add("RHI");
-		//PrivateDependencyModuleNames.Add("ColorManagement");
 		PrivateDependencyModuleNames.Add("Renderer");
 		
 		AddEngineThirdPartyPrivateStaticDependencies(Target, "BinkMediaPlayerSDK");
diff --git a/Engine/Plugins/Media/BinkMedia/Source/BinkMediaPlayer/Private/binktexturesRHI.cpp b/Engine/Plugins/Media/BinkMedia/Source/BinkMediaPlayer/Private/binktexturesRHI.cpp
--- a/Engine/Plugins/Media/BinkMedia/Source/BinkMediaPlayer/Private/binktexturesRHI.cpp
+++ b/Engine/Plugins/Media/BinkMedia/Source/BinkMediaPlayer/Private/binktexturesRHI.cpp
@@ -2,9 +2,9 @@
 
 #define BINKRHIFUNCTIONS
 #define BINKTEXTURESCLEANUP
-#include "../../BinkMediaPlayerSDK/include/egttypes.h"
-#include "../../BinkMediaPlayerSDK/include/binktiny.h"
-#include "../../BinkMediaPlayerSDK/include/binktextures.h"
+#include "egttypes.h"
+#include "binktiny.h"
+#include "binktextures.h"
 
 #include "BinkShaders.h"
 #include "PipelineStateCache.h"
