Bink 2025.01 - new platform SDKs, bug fixes.

diff --git a/Engine/Plugins/Media/BinkMedia/Source/BinkMediaPlayerSDK/include/egttypes.h b/Engine/Plugins/Media/BinkMedia/Source/BinkMediaPlayerSDK/include/egttypes.h
--- a/Engine/Plugins/Media/BinkMedia/Source/BinkMediaPlayerSDK/include/egttypes.h
+++ b/Engine/Plugins/Media/BinkMedia/Source/BinkMediaPlayerSDK/include/egttypes.h
@@ -13,7 +13,7 @@
 #define __RADTYPESH__
 #define __RADRR_COREH__ // block old rr_core
 
-#define RADCOPYRIGHT "Copyright (C) 1994-2024, Epic Games Tools, LLC"
+#define RADCOPYRIGHT "Copyright (C) 1994-2025, Epic Games Tools, LLC"
 
 #if !defined(__RADRES__) // don't include anything for resource compiles
 
