Bink: Improved triggering of OnMediaReachedEnd event

#rnx

diff --git a/Engine/Plugins/Media/BinkMedia/Source/BinkMediaPlayer/Private/Assets/BinkMediaPlayer.cpp b/Engine/Plugins/Media/BinkMedia/Source/BinkMediaPlayer/Private/Assets/BinkMediaPlayer.cpp
--- a/Engine/Plugins/Media/BinkMedia/Source/BinkMediaPlayer/Private/Assets/BinkMediaPlayer.cpp
+++ b/Engine/Plugins/Media/BinkMedia/Source/BinkMediaPlayer/Private/Assets/BinkMediaPlayer.cpp
@@ -168,6 +168,17 @@ bool UBinkMediaPlayer::IsPlaying() const
     return false;
 }
 
+bool UBinkMediaPlayer::IsGotoing() const
+{
+	if (bnk)
+	{
+		BINKPLUGININFO bpinfo = {};
+		BinkPluginInfo(bnk, &bpinfo);
+		return bpinfo.PlaybackState == 2;
+	}
+	return false;
+}
+
 bool UBinkMediaPlayer::OpenUrl( const FString& NewUrl ) 
 {
 	if (NewUrl.IsEmpty()) 
@@ -497,7 +508,7 @@ void UBinkMediaPlayer::HandleMediaPlayerMediaOpened( FString OpenedUrl )
 void UBinkMediaPlayer::Tick(float DeltaTime) 
 {
 	// Check for if we should issue the reached end event
-	if (bnk && !reached_end && !IsPlaying() && !IsPaused() && !IsLooping()) 
+	if (bnk && !reached_end && !IsPlaying() && !IsGotoing() && !IsPaused() && !IsLooping()) 
 	{
 		OnMediaReachedEnd.Broadcast();
 		reached_end = true;
diff --git a/Engine/Plugins/Media/BinkMedia/Source/BinkMediaPlayer/Public/BinkMediaPlayer.h b/Engine/Plugins/Media/BinkMedia/Source/BinkMediaPlayer/Public/BinkMediaPlayer.h
--- a/Engine/Plugins/Media/BinkMedia/Source/BinkMediaPlayer/Public/BinkMediaPlayer.h
+++ b/Engine/Plugins/Media/BinkMedia/Source/BinkMediaPlayer/Public/BinkMediaPlayer.h
@@ -201,6 +201,15 @@ public:
 	UFUNCTION(BlueprintCallable, Category="Media|BinkMediaPlayer")
 	bool IsStopped() const;
 
+	/**
+	 * Checks whether playback is "GoTo'ing".
+	 *
+	 * @return true if playback is gotoing, false otherwise.
+	 * @see IsPaused, IsPlaying, Stop
+	 */
+	UFUNCTION(BlueprintCallable, Category="Media|BinkMediaPlayer")
+	bool IsGotoing() const;
+
 	/**
 	 * Opens the specified media URL.
 	 *
