[bink] Support WinArm64 platforms. (still need libs)

diff --git a/Engine/Plugins/Media/BinkMedia/Source/BinkMediaPlayerSDK/BinkMediaPlayerSDK.Build.cs b/Engine/Plugins/Media/BinkMedia/Source/BinkMediaPlayerSDK/BinkMediaPlayerSDK.Build.cs
--- a/Engine/Plugins/Media/BinkMedia/Source/BinkMediaPlayerSDK/BinkMediaPlayerSDK.Build.cs
+++ b/Engine/Plugins/Media/BinkMedia/Source/BinkMediaPlayerSDK/BinkMediaPlayerSDK.Build.cs
@@ -59,7 +59,14 @@ public class BinkMediaPlayerSDK : ModuleRules
 		{
 			if (Target.IsInPlatformGroup(UnrealPlatformGroup.Microsoft))
 			{
-				Lib = "BinkUnreal" + Platform + ".lib";
+				if (Target.Architecture == UnrealArch.Arm64)
+				{
+					Lib = "BinkUnrealWinArm64.lib";
+				}
+				else
+				{
+					Lib = "BinkUnreal" + Platform + ".lib";
+				}
 			}
 			else if (Target.Platform == UnrealTargetPlatform.Mac)
 			{
