Bink: Fixed MT issue in PIE mode switch

#rnx

diff --git a/Engine/Plugins/Media/BinkMedia/Source/BinkMediaPlayer/Private/BinkFunctionLibrary.cpp b/Engine/Plugins/Media/BinkMedia/Source/BinkMediaPlayer/Private/BinkFunctionLibrary.cpp
--- a/Engine/Plugins/Media/BinkMedia/Source/BinkMediaPlayer/Private/BinkFunctionLibrary.cpp
+++ b/Engine/Plugins/Media/BinkMedia/Source/BinkMediaPlayer/Private/BinkFunctionLibrary.cpp
@@ -13,25 +13,18 @@
 extern TSharedPtr<FBinkMovieStreamer, ESPMode::ThreadSafe> MovieStreamer;
 
 // Note: Has to be in a seperate function because you can't do #if inside a render command macro
-static void Bink_DrawOverlays_Internal(FRHICommandListImmediate &RHICmdList, const UGameViewportClient* gameViewport) {
-	if (!GEngine || !gameViewport || !gameViewport->Viewport) {
-		return;
-	}
-
-	FVector2D screenSize;
-	gameViewport->GetViewportSize(screenSize);
-	const FTextureRHIRef &backbuffer = gameViewport->Viewport->GetRenderTargetTexture();
-	if(!backbuffer.GetReference()) 
+static void Bink_DrawOverlays_Internal(FRHICommandListImmediate &RHICmdList, FTextureRHIRef BackBuffer, FVector2D ScreenSize) {
+	if(!BackBuffer.GetReference()) 
 	{
 		return;
 	}
 
 	BINKPLUGINFRAMEINFO FrameInfo = {};
-	FrameInfo.screen_resource = backbuffer.GetReference();
+	FrameInfo.screen_resource = BackBuffer.GetReference();
 	FrameInfo.screen_resource_state = 4; // D3D12_RESOURCE_STATE_RENDER_TARGET; (only used in d3d12)
-	FrameInfo.width = screenSize.X;
-	FrameInfo.height = screenSize.Y;
-	FrameInfo.sdr_or_hdr = backbuffer->GetFormat() == PF_A2B10G10R10 ? 1 : 0;
+	FrameInfo.width = ScreenSize.X;
+	FrameInfo.height = ScreenSize.Y;
+	FrameInfo.sdr_or_hdr = BackBuffer->GetFormat() == PF_A2B10G10R10 ? 1 : 0;
 	FrameInfo.cmdBuf = &RHICmdList;
 	BinkPluginSetPerFrameInfo(&FrameInfo);
 	BinkPluginAllScheduled();
@@ -40,17 +33,21 @@ static void Bink_DrawOverlays_Internal(FRHICommandListImmediate &RHICmdList, con
 
 void UBinkFunctionLibrary::Bink_DrawOverlays() 
 {
-	TWeakObjectPtr<UGameViewportClient> gameViewport = (GEngine != nullptr) ? GEngine->GameViewport : nullptr;
-	if (!gameViewport.IsValid()) {
+	if (!GEngine || !GEngine->GameViewport || !GEngine->GameViewport->Viewport) {
+		return;
+	}
+
+	FVector2D ScreenSize;
+	GEngine->GameViewport->GetViewportSize(ScreenSize);
+	FTextureRHIRef BackBuffer = GEngine->GameViewport->Viewport->GetRenderTargetTexture();
+
+	if (!BackBuffer.IsValid()) {
 		return;
 	}
 
-	ENQUEUE_RENDER_COMMAND(BinkOverlays)([gameViewport](FRHICommandListImmediate& RHICmdList) 
+	ENQUEUE_RENDER_COMMAND(BinkOverlays)([ScreenSize, BackBuffer](FRHICommandListImmediate& RHICmdList) 
 	{ 
-		if (gameViewport.IsValid())
-		{
-			Bink_DrawOverlays_Internal(RHICmdList, gameViewport.Get());
-		}
+		Bink_DrawOverlays_Internal(RHICmdList, BackBuffer, ScreenSize);
 	});
 }
 
