Transition Bink Media Texture back to SRV instead of UAV

diff --git a/Engine/Plugins/Media/BinkMedia/Source/BinkMediaPlayer/Private/Assets/BinkMediaTextureResource.cpp b/Engine/Plugins/Media/BinkMedia/Source/BinkMediaPlayer/Private/Assets/BinkMediaTextureResource.cpp
--- a/Engine/Plugins/Media/BinkMedia/Source/BinkMediaPlayer/Private/Assets/BinkMediaTextureResource.cpp
+++ b/Engine/Plugins/Media/BinkMedia/Source/BinkMediaPlayer/Private/Assets/BinkMediaTextureResource.cpp
@@ -74,7 +74,7 @@ void FBinkMediaTextureResource::InitRHI(FRHICommandListBase& RHICmdListBase)
 		RHICmdList.BeginRenderPass(RPInfo, TEXT("ClearTexture"));
 		RHICmdList.EndRenderPass();
 		RHICmdList.SetViewport(0, 0, 0, w, h, 1);
-		RHICmdList.Transition(FRHITransitionInfo(TextureRHI.GetReference(), ERHIAccess::RTV, ERHIAccess::UAVGraphics));
+		RHICmdList.Transition(FRHITransitionInfo(TextureRHI.GetReference(), ERHIAccess::RTV, ERHIAccess::SRVGraphics));
 	}
 }
 
@@ -122,7 +122,7 @@ void FBinkMediaTextureResource::Clear()
 		RHICmdList.BeginRenderPass(RPInfo, TEXT("ClearTexture"));
 		RHICmdList.EndRenderPass();
 		RHICmdList.SetViewport(0, 0, 0, w, h, 1);
-		RHICmdList.Transition(FRHITransitionInfo(ref2.GetReference(), ERHIAccess::RTV, ERHIAccess::UAVGraphics));
+		RHICmdList.Transition(FRHITransitionInfo(ref2.GetReference(), ERHIAccess::RTV, ERHIAccess::SRVGraphics));
 	});
 }
 
