Bink 2025.05 binaries (new platform sdks)

diff --git a/Engine/Plugins/Media/BinkMedia/Source/BinkMediaPlayerSDK/include/egttypes.h b/Engine/Plugins/Media/BinkMedia/Source/BinkMediaPlayerSDK/include/egttypes.h
--- a/Engine/Plugins/Media/BinkMedia/Source/BinkMediaPlayerSDK/include/egttypes.h
+++ b/Engine/Plugins/Media/BinkMedia/Source/BinkMediaPlayerSDK/include/egttypes.h
@@ -186,7 +186,7 @@
 
 
 #if !defined(__RADDETECTED__)
-  #error "radtypes.h did not detect your platform."
+  #error "egttypes.h did not detect your platform."
 #endif
 
 // ========================================================
