* Removed bAllowConfidentialPlatformDefines from modules because they don't need it (android and ios are not confidential platforms)
* Added bAllowConfidentialPlatformDefines to D3D12RHI which actually needed it

diff --git a/Engine/Plugins/Media/BinkMedia/Source/BinkMediaPlayer/BinkMediaPlayer.Build.cs b/Engine/Plugins/Media/BinkMedia/Source/BinkMediaPlayer/BinkMediaPlayer.Build.cs
--- a/Engine/Plugins/Media/BinkMedia/Source/BinkMediaPlayer/BinkMediaPlayer.Build.cs
+++ b/Engine/Plugins/Media/BinkMedia/Source/BinkMediaPlayer/BinkMediaPlayer.Build.cs
@@ -9,8 +9,6 @@ public class BinkMediaPlayer : ModuleRules
 {
 	public BinkMediaPlayer(ReadOnlyTargetRules Target) : base(Target)
 	{
-		bAllowConfidentialPlatformDefines = true;
-		
 		PublicDependencyModuleNames.Add("Core");
 		PublicDependencyModuleNames.Add("CoreUObject");
 		PublicDependencyModuleNames.Add("Engine");
diff --git a/Engine/Plugins/Media/BinkMedia/Source/BinkMediaPlayerSDK/BinkMediaPlayerSDK.Build.cs b/Engine/Plugins/Media/BinkMedia/Source/BinkMediaPlayerSDK/BinkMediaPlayerSDK.Build.cs
--- a/Engine/Plugins/Media/BinkMedia/Source/BinkMediaPlayerSDK/BinkMediaPlayerSDK.Build.cs
+++ b/Engine/Plugins/Media/BinkMedia/Source/BinkMediaPlayerSDK/BinkMediaPlayerSDK.Build.cs
@@ -20,8 +20,6 @@ public class BinkMediaPlayerSDK : ModuleRules
     {
 		Type = ModuleType.External;
 
-		bAllowConfidentialPlatformDefines = true;
-
         PublicDependencyModuleNames.Add("Core");
         PublicDependencyModuleNames.Add("CoreUObject");
         PublicDependencyModuleNames.Add("Engine");
diff --git a/Engine/Plugins/MemoryUsageQueries/Source/MemoryUsageQueries/MemoryUsageQueries.Build.cs b/Engine/Plugins/MemoryUsageQueries/Source/MemoryUsageQueries/MemoryUsageQueries.Build.cs
--- a/Engine/Plugins/MemoryUsageQueries/Source/MemoryUsageQueries/MemoryUsageQueries.Build.cs
+++ b/Engine/Plugins/MemoryUsageQueries/Source/MemoryUsageQueries/MemoryUsageQueries.Build.cs
@@ -7,7 +7,6 @@ namespace UnrealBuildTool.Rules
 		public MemoryUsageQueries(ReadOnlyTargetRules Target) : base(Target)
 		{
 			PCHUsage = PCHUsageMode.UseExplicitOrSharedPCHs;
-			bAllowConfidentialPlatformDefines = true;
 
 			PublicIncludePaths.AddRange(
 				new string[] {
diff --git a/Engine/Plugins/Mutable/Source/CustomizableObject/CustomizableObject.Build.cs b/Engine/Plugins/Mutable/Source/CustomizableObject/CustomizableObject.Build.cs
--- a/Engine/Plugins/Mutable/Source/CustomizableObject/CustomizableObject.Build.cs
+++ b/Engine/Plugins/Mutable/Source/CustomizableObject/CustomizableObject.Build.cs
@@ -10,7 +10,6 @@ public class CustomizableObject : ModuleRules
 
 		DefaultBuildSettings = BuildSettingsVersion.V2;
 
-		bAllowConfidentialPlatformDefines = true;
 		//bUseUnity = false;
 
 		PrivateDependencyModuleNames.AddRange([
diff --git a/Engine/Source/Runtime/D3D12RHI/D3D12RHI.Build.cs b/Engine/Source/Runtime/D3D12RHI/D3D12RHI.Build.cs
--- a/Engine/Source/Runtime/D3D12RHI/D3D12RHI.Build.cs
+++ b/Engine/Source/Runtime/D3D12RHI/D3D12RHI.Build.cs
@@ -8,6 +8,8 @@ public class D3D12RHI : ModuleRules
 {
 	public D3D12RHI(ReadOnlyTargetRules Target) : base(Target)
 	{
+		bAllowConfidentialPlatformDefines = true;
+
 		PrivateDependencyModuleNames.AddAll(
 			"CoreUObject",
 			"Engine",
diff --git a/Engine/Source/Runtime/Experimental/IndexedCacheStorage/IndexedCacheStorage.Build.cs b/Engine/Source/Runtime/Experimental/IndexedCacheStorage/IndexedCacheStorage.Build.cs
--- a/Engine/Source/Runtime/Experimental/IndexedCacheStorage/IndexedCacheStorage.Build.cs
+++ b/Engine/Source/Runtime/Experimental/IndexedCacheStorage/IndexedCacheStorage.Build.cs
@@ -8,7 +8,6 @@ public class IndexedCacheStorage : ModuleRules
 	{
 		PublicDependencyModuleNames.Add("Core");
 
-		bAllowConfidentialPlatformDefines = true;
 		CppCompileWarningSettings.UnsafeTypeCastWarningLevel = WarningLevel.Error; 
 	}
 }
diff --git a/Engine/Source/Runtime/Experimental/IoStore/OnDemand/IoStoreOnDemand.Build.cs b/Engine/Source/Runtime/Experimental/IoStore/OnDemand/IoStoreOnDemand.Build.cs
--- a/Engine/Source/Runtime/Experimental/IoStore/OnDemand/IoStoreOnDemand.Build.cs
+++ b/Engine/Source/Runtime/Experimental/IoStore/OnDemand/IoStoreOnDemand.Build.cs
@@ -19,7 +19,6 @@ public class IoStoreOnDemand : ModuleRules
 			}
 		);
 
-		bAllowConfidentialPlatformDefines = true;
 		CppCompileWarningSettings.UnsafeTypeCastWarningLevel = WarningLevel.Error; 
 
 		if (Target.Platform.IsInGroup(UnrealPlatformGroup.Desktop) &&
diff --git a/Engine/Source/Runtime/Experimental/IoStore/OnDemandCore/IoStoreOnDemandCore.Build.cs b/Engine/Source/Runtime/Experimental/IoStore/OnDemandCore/IoStoreOnDemandCore.Build.cs
--- a/Engine/Source/Runtime/Experimental/IoStore/OnDemandCore/IoStoreOnDemandCore.Build.cs
+++ b/Engine/Source/Runtime/Experimental/IoStore/OnDemandCore/IoStoreOnDemandCore.Build.cs
@@ -14,7 +14,6 @@ public class IoStoreOnDemandCore : ModuleRules
 			}
 		);
 
-		bAllowConfidentialPlatformDefines = true;
 		CppCompileWarningSettings.UnsafeTypeCastWarningLevel = WarningLevel.Error; 
 	}
 }
