Bink: Fixed CookOnTheFly paths (long standing bug, thanks BugHawk)

diff --git a/Engine/Plugins/Media/BinkMedia/Source/BinkMediaPlayer/Private/BinkMediaPlayerModule.cpp b/Engine/Plugins/Media/BinkMedia/Source/BinkMediaPlayer/Private/BinkMediaPlayerModule.cpp
--- a/Engine/Plugins/Media/BinkMedia/Source/BinkMediaPlayer/Private/BinkMediaPlayerModule.cpp
+++ b/Engine/Plugins/Media/BinkMedia/Source/BinkMediaPlayer/Private/BinkMediaPlayerModule.cpp
@@ -172,7 +172,7 @@ FString BinkUE4CookOnTheFlyPath(FString path, const TCHAR *filename)
 	toPath = toPath.Replace(TEXT("/./"), TEXT("/"));
 	fromPath = fromPath.Replace(TEXT("/./"), TEXT("/"));
 	FPlatformFileManager::Get().GetPlatformFile().CopyFile(*toPath, *fromPath);
-	return fromPath;
+	return toPath;
 }
 
 struct FBinkMediaPlayerModule : IModuleInterface, FTickableGameObject
