Fixed return value of movie streamers to be the result of the actual open call.

#rnx

diff --git a/Engine/Plugins/Media/BinkMedia/Source/BinkMediaPlayer/Private/BinkMovieStreamer.cpp b/Engine/Plugins/Media/BinkMedia/Source/BinkMediaPlayer/Private/BinkMovieStreamer.cpp
--- a/Engine/Plugins/Media/BinkMedia/Source/BinkMediaPlayer/Private/BinkMovieStreamer.cpp
+++ b/Engine/Plugins/Media/BinkMedia/Source/BinkMediaPlayer/Private/BinkMovieStreamer.cpp
@@ -1,5 +1,5 @@
 // Copyright Epic Games Tools LLC
-//   Licenced under the Unreal Engine EULA 
+//   Licenced under the Unreal Engine EULA
 
 #include "BinkMovieStreamer.h"
 
@@ -16,14 +16,14 @@
 
 DEFINE_LOG_CATEGORY(LogBinkMoviePlayer);
 
-FBinkMovieStreamer::FBinkMovieStreamer() 
+FBinkMovieStreamer::FBinkMovieStreamer()
 	: bnk()
 {
 	MovieViewport = MakeShareable(new FMovieViewport());
 	PlaybackType = MT_Normal;
 }
 
-FBinkMovieStreamer::~FBinkMovieStreamer() 
+FBinkMovieStreamer::~FBinkMovieStreamer()
 {
 	CloseMovie();
 	Cleanup();
@@ -32,50 +32,49 @@ FBinkMovieStreamer::~FBinkMovieStreamer()
 	TextureFreeList.Empty();
 }
 
-bool FBinkMovieStreamer::Init(const TArray<FString>& MoviePaths, TEnumAsByte<EMoviePlaybackType> inPlaybackType) 
+bool FBinkMovieStreamer::Init(const TArray<FString>& MoviePaths, TEnumAsByte<EMoviePlaybackType> inPlaybackType)
 {
-	if (MoviePaths.Num() == 0) 
+	if (MoviePaths.Num() == 0)
 	{
 		return false;
 	}
 	MovieIndex = -1;
 	PlaybackType = inPlaybackType;
 	StoredMoviePaths = MoviePaths;
-	OpenNextMovie();
-	return true;
+	return OpenNextMovie();
 }
 
-FString FBinkMovieStreamer::GetMovieName() 
+FString FBinkMovieStreamer::GetMovieName()
 {
 	return StoredMoviePaths.IsValidIndex(MovieIndex) ? StoredMoviePaths[MovieIndex] : TEXT("");
 }
 
-bool FBinkMovieStreamer::IsLastMovieInPlaylist() 
+bool FBinkMovieStreamer::IsLastMovieInPlaylist()
 {
 	return MovieIndex == StoredMoviePaths.Num() -1;
 }
 
-void FBinkMovieStreamer::ForceCompletion() 
+void FBinkMovieStreamer::ForceCompletion()
 {
 	CloseMovie();
 }
 
-bool FBinkMovieStreamer::Tick(FRHICommandListBase& RHICmdList, float DeltaTime) 
+bool FBinkMovieStreamer::Tick(FRHICommandListBase& RHICmdList, float DeltaTime)
 {
 	// Loops through all movies while not valid
-	if (!bnk) 
+	if (!bnk)
 	{
 		CloseMovie();
-		if (MovieIndex < StoredMoviePaths.Num() - 1) 
+		if (MovieIndex < StoredMoviePaths.Num() - 1)
 		{
 			OpenNextMovie();
-		} 
-		else if (PlaybackType != MT_Normal) 
+		}
+		else if (PlaybackType != MT_Normal)
 		{
 			MovieIndex = PlaybackType == MT_LoadingLoop ? StoredMoviePaths.Num() - 2 : -1;
 			OpenNextMovie();
-		} 
-		else 
+		}
+		else
 		{
 			return true;
 		}
@@ -83,35 +82,35 @@ bool FBinkMovieStreamer::Tick(FRHICommandListBase& RHICmdList, float DeltaTime)
 
 	BINKPLUGININFO bpinfo = {};
 	BinkPluginInfo(bnk, &bpinfo);
-	if (bpinfo.PlaybackState == 3) 
+	if (bpinfo.PlaybackState == 3)
 	{
 		CloseMovie();
-		if (MovieIndex < StoredMoviePaths.Num() - 1) 
+		if (MovieIndex < StoredMoviePaths.Num() - 1)
 		{
 			OpenNextMovie();
-		} 
-		else if (PlaybackType != MT_Normal) 
+		}
+		else if (PlaybackType != MT_Normal)
 		{
 			MovieIndex = PlaybackType == MT_LoadingLoop ? StoredMoviePaths.Num() - 2 : -1;
 			OpenNextMovie();
-		} 
-		else 
+		}
+		else
 		{
 			return true;
 		}
 	}
-	if (!bnk) 
+	if (!bnk)
 	{
 		return false;
 	}
 
 	FSlateTexture2DRHIRef* CurrentTexture = Texture.Get();
-	if(CurrentTexture) 
+	if(CurrentTexture)
 	{
-		FVector2D destUpperLeft = GetDefault<UBinkMoviePlayerSettings>()->BinkDestinationUpperLeft; 
+		FVector2D destUpperLeft = GetDefault<UBinkMoviePlayerSettings>()->BinkDestinationUpperLeft;
 		FVector2D destLowerRight = GetDefault<UBinkMoviePlayerSettings>()->BinkDestinationLowerRight;
 
-		if( !CurrentTexture->IsInitialized() ) 
+		if( !CurrentTexture->IsInitialized() )
 		{
 			CurrentTexture->InitResource(RHICmdList);
 		}
@@ -171,22 +170,22 @@ bool FBinkMovieStreamer::Tick(FRHICommandListBase& RHICmdList, float DeltaTime)
 		BinkPluginProcessBinks(0);
 		BinkPluginDraw(1, 0);
 
-		MovieViewport->SetTexture(Texture);	
+		MovieViewport->SetTexture(Texture);
 	}
 
 	return false;
 }
 
-void FBinkMovieStreamer::Cleanup() 
+void FBinkMovieStreamer::Cleanup()
 {
 	FlushRenderingCommands();
-	for( int32 TextureIndex = 0; TextureIndex < TextureFreeList.Num(); ++TextureIndex ) 
+	for( int32 TextureIndex = 0; TextureIndex < TextureFreeList.Num(); ++TextureIndex )
 	{
 		BeginReleaseResource( TextureFreeList[TextureIndex].Get() );
 	}
 }
 
-bool FBinkMovieStreamer::OpenNextMovie() 
+bool FBinkMovieStreamer::OpenNextMovie()
 {
 	MovieIndex++;
 	check(StoredMoviePaths.Num() > 0 && MovieIndex < StoredMoviePaths.Num());
@@ -201,11 +200,11 @@ bool FBinkMovieStreamer::OpenNextMovie()
 
 	IPlatformFile& PlatformFile = FPlatformFileManager::Get().GetPlatformFile();
 
-	FString MoviePathTbl[] = 
+	FString MoviePathTbl[] =
 	{
 		StoredMoviePaths[MovieIndex] + TEXT(".bk2")
 	};
-	for (int i = 0; i < sizeof(MoviePathTbl) / sizeof(MoviePathTbl[0]) && !bnk; ++i) 
+	for (int i = 0; i < sizeof(MoviePathTbl) / sizeof(MoviePathTbl[0]) && !bnk; ++i)
 	{
 		FString FullMoviePath = BINKMOVIEPATH + MoviePathTbl[i];
 #if PLATFORM_ANDROID
@@ -217,9 +216,9 @@ bool FBinkMovieStreamer::OpenNextMovie()
 		int64 FileOffset = IAndroidPlatformFile::GetPlatformPhysical().FileStartOffset(*FullMoviePath);
 		FString FileRootPath = IAndroidPlatformFile::GetPlatformPhysical().FileRootPath(*FullMoviePath);
 
-		if (IAndroidPlatformFile::GetPlatformPhysical().IsAsset(*FullMoviePath)) 
+		if (IAndroidPlatformFile::GetPlatformPhysical().IsAsset(*FullMoviePath))
 		{
-			if (JNIEnv* env = FAndroidApplication::GetJavaEnv()) 
+			if (JNIEnv* env = FAndroidApplication::GetJavaEnv())
 			{
 				extern struct android_app* GNativeAndroidApp;
 				jclass clazz = env->GetObjectClass(GNativeAndroidApp->activity->clazz);
@@ -231,27 +230,27 @@ bool FBinkMovieStreamer::OpenNextMovie()
 				env->ReleaseStringUTFChars((jstring)result, apkPath);
 			}
 		}
-		else 
+		else
 		{
 			bnk = BinkPluginOpen(TCHAR_TO_ANSI(*FileRootPath), soundTrack, soundTrackStart, bufferMode, FileOffset);
 		}
 #else
 		FString ExternalPath = PlatformFile.ConvertToAbsolutePathForExternalAppForRead(*FullMoviePath);
-		if (FPlatformFileManager::Get().GetPlatformFile().FileExists(*ExternalPath)) 
+		if (FPlatformFileManager::Get().GetPlatformFile().FileExists(*ExternalPath))
 		{
-#if PLATFORM_WINDOWS 
+#if PLATFORM_WINDOWS
 			bnk = BinkPluginOpenUTF16((unsigned short*)*ExternalPath, soundTrack, soundTrackStart, bufferMode, 0);
 #else
 			bnk = BinkPluginOpen(TCHAR_TO_ANSI(*ExternalPath), soundTrack, soundTrackStart, bufferMode, 0);
 #endif
 		}
-		if (!bnk) 
+		if (!bnk)
 		{
 			FString CookPath = *BinkUE4CookOnTheFlyPath(FPaths::ConvertRelativePathToFull(BINKMOVIEPATH), *MoviePathTbl[i]);
 			ExternalPath = PlatformFile.ConvertToAbsolutePathForExternalAppForRead(*FullMoviePath);
-			if (FPlatformFileManager::Get().GetPlatformFile().FileExists(*ExternalPath)) 
+			if (FPlatformFileManager::Get().GetPlatformFile().FileExists(*ExternalPath))
 			{
-#if PLATFORM_WINDOWS 
+#if PLATFORM_WINDOWS
 				bnk = BinkPluginOpenUTF16((unsigned short *)*ExternalPath, soundTrack, soundTrackStart, bufferMode, 0);
 #else
 				bnk = BinkPluginOpen(TCHAR_TO_ANSI(*ExternalPath), soundTrack, soundTrackStart, bufferMode, 0);
@@ -260,7 +259,7 @@ bool FBinkMovieStreamer::OpenNextMovie()
 		}
 #endif
 	}
-	if (!bnk) 
+	if (!bnk)
 	{
 		return false;
 	}
@@ -271,30 +270,30 @@ bool FBinkMovieStreamer::OpenNextMovie()
 	VideoDimensions.X = bpinfo.Width;
 	VideoDimensions.Y = bpinfo.Height;
 
-	if( VideoDimensions.X > 0 && VideoDimensions.Y > 0 ) 
+	if( VideoDimensions.X > 0 && VideoDimensions.Y > 0 )
 	{
 		EPixelFormat pixelFmt = bink_force_pixel_format == PF_Unknown ? (EPixelFormat)GetDefault<UBinkMoviePlayerSettings>()->BinkPixelFormat : bink_force_pixel_format;
-		if( TextureFreeList.Num() > 0 ) 
+		if( TextureFreeList.Num() > 0 )
 		{
 			Texture = TextureFreeList.Pop();
 
-			if( Texture->GetWidth() != VideoDimensions.X || Texture->GetHeight() != VideoDimensions.Y ) 
+			if( Texture->GetWidth() != VideoDimensions.X || Texture->GetHeight() != VideoDimensions.Y )
 			{
 				FSlateTexture2DRHIRef *ref = Texture.Get();
-				ENQUEUE_RENDER_COMMAND(UpdateMovieTexture)([ref,VideoDimensions](FRHICommandListImmediate& RHICmdList) 
+				ENQUEUE_RENDER_COMMAND(UpdateMovieTexture)([ref,VideoDimensions](FRHICommandListImmediate& RHICmdList)
 				{
 					ref->Resize( VideoDimensions.X, VideoDimensions.Y );
 				});
 			}
-		} 
-		else 
+		}
+		else
 		{
 			const bool bCreateEmptyTexture = true;
 			Texture = MakeShareable(new FSlateTexture2DRHIRef(VideoDimensions.X, VideoDimensions.Y, pixelFmt, NULL, TexCreate_RenderTargetable, bCreateEmptyTexture));
 			FSlateTexture2DRHIRef *ref = Texture.Get();
 
-			ENQUEUE_RENDER_COMMAND(InitMovieTexture)([ref](FRHICommandListImmediate& RHICmdList) 
-			{ 
+			ENQUEUE_RENDER_COMMAND(InitMovieTexture)([ref](FRHICommandListImmediate& RHICmdList)
+			{
 				ref->InitResource(RHICmdList);
 			});
 		}
@@ -302,20 +301,20 @@ bool FBinkMovieStreamer::OpenNextMovie()
 	return true;
 }
 
-void FBinkMovieStreamer::CloseMovie() 
+void FBinkMovieStreamer::CloseMovie()
 {
-	if (GetMoviePlayer() != nullptr) 
+	if (GetMoviePlayer() != nullptr)
 	{
 		BroadcastCurrentMovieClipFinished(GetMovieName());
 	}
 
-	if (Texture.IsValid()) 
+	if (Texture.IsValid())
 	{
 		TextureFreeList.Add(Texture);
 		MovieViewport->SetTexture(NULL);
 		Texture.Reset();
 	}
-	if(bnk) 
+	if(bnk)
 	{
 		BinkPluginClose(bnk);
 		bnk = NULL;
diff --git a/Engine/Plugins/Runtime/WindowsMoviePlayer/Source/WindowsMoviePlayer/Private/WindowsMovieStreamer.cpp b/Engine/Plugins/Runtime/WindowsMoviePlayer/Source/WindowsMoviePlayer/Private/WindowsMovieStreamer.cpp
--- a/Engine/Plugins/Runtime/WindowsMoviePlayer/Source/WindowsMoviePlayer/Private/WindowsMovieStreamer.cpp
+++ b/Engine/Plugins/Runtime/WindowsMoviePlayer/Source/WindowsMoviePlayer/Private/WindowsMovieStreamer.cpp
@@ -63,9 +63,7 @@ bool FMediaFoundationMovieStreamer::Init(const TArray<FString>& MoviePaths, TEnu
 
 	MovieViewport->SetTexture(nullptr);
 
-	OpenNextMovie();
-
-	return true;
+	return OpenNextMovie();
 }
 
 void FMediaFoundationMovieStreamer::ForceCompletion()
@@ -238,7 +236,7 @@ bool FMediaFoundationMovieStreamer::Tick(FRHICommandListBase& RHICmdList, float
 	return false;
 }
 
-FString FMediaFoundationMovieStreamer::GetMovieName() 
+FString FMediaFoundationMovieStreamer::GetMovieName()
 {
 	return StoredMoviePaths.IsValidIndex(MovieIndex) ? StoredMoviePaths[MovieIndex] : TEXT("");
 }
@@ -260,7 +258,7 @@ bool FMediaFoundationMovieStreamer::OpenNextMovie()
 	FString MoviePath = FPaths::ProjectContentDir() + TEXT("Movies/") + StoredMoviePaths[MovieIndex];
 
 	SampleGrabberCallback = new FSampleGrabberCallback(TextureData);
-	
+
 	VideoPlayer = new FVideoPlayer();
 	FIntPoint VideoDimensions = VideoPlayer->OpenFile(MoviePath, SampleGrabberCallback, bUseSound);
 
@@ -405,7 +403,7 @@ HRESULT FVideoPlayer::Invoke(IMFAsyncResult* AsyncResult)
 	case MESessionClosed:
 		bFinishedAndClose = true;
 		break;
-	
+
 	case MESessionTopologySet:
 		if (!SUCCEEDED(EventStatus))
 		{
@@ -455,7 +453,7 @@ HRESULT FVideoPlayer::Invoke(IMFAsyncResult* AsyncResult)
 			return S_OK;
 		}
 	}
-	
+
 	return S_OK;
 }
 
@@ -489,7 +487,7 @@ FIntPoint FVideoPlayer::OpenFile(const FString& FilePath, class FSampleGrabberCa
 		{
 			HResult = Source->QueryInterface(IID_PPV_ARGS(&MediaSource));
 			Source->Release();
-			
+
 			OutDimensions = SetPlaybackTopology(SampleGrabberCallback, bUseSound);
 		}
 		else
@@ -497,7 +495,7 @@ FIntPoint FVideoPlayer::OpenFile(const FString& FilePath, class FSampleGrabberCa
 			UE_LOG(LogWindowsMoviePlayer, Log, TEXT("Unable to load movie: %s"), *PathPlusExt );
 			MovieIsFinished.Set(1);
 		}
-		
+
 	}
 
 	return OutDimensions;
@@ -548,11 +546,11 @@ FIntPoint FVideoPlayer::SetPlaybackTopology(FSampleGrabberCallback* SampleGrabbe
 	FIntPoint OutDimensions = FIntPoint(ForceInit);
 
 	HRESULT HResult = S_OK;
-	
+
 	IMFPresentationDescriptor* PresentationDesc = NULL;
 	HResult = MediaSource->CreatePresentationDescriptor(&PresentationDesc);
 	check(SUCCEEDED(HResult));
-	
+
 	IMFTopology* Topology = NULL;
 	HResult = MFCreateTopology(&Topology);
 	check(SUCCEEDED(HResult));
@@ -579,7 +577,7 @@ FIntPoint FVideoPlayer::SetPlaybackTopology(FSampleGrabberCallback* SampleGrabbe
 
 		StreamDesc->Release();
 	}
-	
+
 	HResult = MediaSession->SetTopology(0, Topology);
 	check(SUCCEEDED(HResult));
 
@@ -594,7 +592,7 @@ FIntPoint FVideoPlayer::AddStreamToTopology(IMFTopology* Topology, IMFPresentati
 	FIntPoint OutDimensions = FIntPoint(ForceInit);
 
 	HRESULT HResult = S_OK;
-	
+
 	IMFActivate* SinkActivate = NULL;
 	{
 		IMFMediaTypeHandler* Handler = NULL;
@@ -637,7 +635,7 @@ FIntPoint FVideoPlayer::AddStreamToTopology(IMFTopology* Topology, IMFPresentati
 		}
 		else if (MajorType == MFMediaType_Video)
 		{
-			
+
 			IMFMediaType* OutputType = NULL;
 			HResult = Handler->GetCurrentMediaType(&OutputType);
 			check(SUCCEEDED(HResult));
@@ -681,21 +679,21 @@ FIntPoint FVideoPlayer::AddStreamToTopology(IMFTopology* Topology, IMFPresentati
 					check(SUCCEEDED(HResult));
 
 					VideoTrackFormat.SampleFormat = EMediaTextureSampleFormat::CharYUY2;
-					
+
 					int32 AlignedOutputX = VideoTrackFormat.OutputDim.X;
-					
-					if ((SourceVideoSubType == MFVideoFormat_H264) || (SourceVideoSubType == MFVideoFormat_H264_ES)) 
+
+					if ((SourceVideoSubType == MFVideoFormat_H264) || (SourceVideoSubType == MFVideoFormat_H264_ES))
 					{
 						AlignedOutputX = Align(AlignedOutputX, 16);
 					}
-					
+
 					int32 SampleStride = AlignedOutputX * 2; // 2 bytes per pixel
 
 					if (SampleStride < 0)
 					{
 						SampleStride = -SampleStride;
 					}
-					
+
 					VideoTrackFormat.BufferDim = FIntPoint(AlignedOutputX / 2, VideoTrackFormat.OutputDim.Y); // 2 pixels per texel
 					VideoTrackFormat.BufferStride = SampleStride;
 				}
@@ -704,7 +702,7 @@ FIntPoint FVideoPlayer::AddStreamToTopology(IMFTopology* Topology, IMFPresentati
 			HResult = InputType->SetUINT32(MF_MT_ALL_SAMPLES_INDEPENDENT, TRUE);
 			check(SUCCEEDED(HResult));
 			HResult = MFCreateSampleGrabberSinkActivate(InputType, SampleGrabberCallback, &SinkActivate);
-		
+
 			check(SUCCEEDED(HResult));
 			SAFE_RELEASE(InputType);
 
@@ -715,7 +713,7 @@ FIntPoint FVideoPlayer::AddStreamToTopology(IMFTopology* Topology, IMFPresentati
 
 		SAFE_RELEASE(Handler);
 	}
-	
+
 	IMFTopologyNode* SourceNode = NULL;
 	{
 		HResult = MFCreateTopologyNode(MF_TOPOLOGY_SOURCESTREAM_NODE, &SourceNode);
@@ -729,7 +727,7 @@ FIntPoint FVideoPlayer::AddStreamToTopology(IMFTopology* Topology, IMFPresentati
 		HResult = Topology->AddNode(SourceNode);
 		check(SUCCEEDED(HResult));
 	}
-	
+
 	IMFTopologyNode* OutputNode = NULL;
 	{
 		HResult = MFCreateTopologyNode(MF_TOPOLOGY_OUTPUT_NODE, &OutputNode);
