Bink: Fixed Vulkan backend. Made it so that texture updates do not happen inside a render pass.

diff --git a/Engine/Plugins/Media/BinkMedia/Source/BinkMediaPlayer/Private/binktexturesRHI.cpp b/Engine/Plugins/Media/BinkMedia/Source/BinkMediaPlayer/Private/binktexturesRHI.cpp
--- a/Engine/Plugins/Media/BinkMedia/Source/BinkMediaPlayer/Private/binktexturesRHI.cpp
+++ b/Engine/Plugins/Media/BinkMedia/Source/BinkMediaPlayer/Private/binktexturesRHI.cpp
@@ -463,7 +463,7 @@ static void Draw_textures(BINKTEXTURES* ptextures, BINKSHADERS* pshaders, void*
 	BinkGraphBuilder.AddPass(
 		MoveTemp(EventName),
 		consts,
-		ERDGPassFlags::Raster,
+		ERDGPassFlags::Raster | ERDGPassFlags::SkipRenderPass,
 		[&](FRHICommandListImmediate& RHICmdList)
 		{
 			if (textures->dirty)
@@ -480,6 +480,11 @@ static void Draw_textures(BINKTEXTURES* ptextures, BINKSHADERS* pshaders, void*
 				if (hasHPlane) RHICmdList.UpdateTexture2D(textures->Htexture[frame_num], 0, region_YAH, bp_src->HPlane.BufferPitch, (uint8*)bp_src->HPlane.Buffer);
 			}
 
+
+			RHICmdList.Transition(FRHITransitionInfo(consts->RenderTargets[0].GetTexture()->GetRHI()->GetTexture2D(), ERHIAccess::Unknown, ERHIAccess::RTV));
+			FRHIRenderPassInfo RPInfo(consts->RenderTargets[0].GetTexture()->GetRHI()->GetTexture2D(), ERenderTargetActions::DontLoad_Store);
+			RHICmdList.BeginRenderPass(RPInfo, TEXT("BinkPass"));
+
 			FGraphicsPipelineStateInitializer GraphicsPSOInit;
 			RHICmdList.ApplyCachedRenderTargets(GraphicsPSOInit);
 
@@ -528,6 +533,8 @@ static void Draw_textures(BINKTEXTURES* ptextures, BINKSHADERS* pshaders, void*
 			SetShaderParameters(RHICmdList, DrawVS, DrawVS.GetVertexShader(), *vert_params);
 
 			RHICmdList.DrawPrimitive(0, 2, 1);
+
+			RHICmdList.EndRenderPass();
 		});
 
 	BinkGraphBuilder.Execute();
