IOS, check file existence based on relative path prior to resolving to absolute path. Additionally, fix check for movie existence in CookPath.

#Synced-CL 48425151

diff --git a/Engine/Plugins/Media/BinkMedia/Source/BinkMediaPlayer/Private/BinkMovieStreamer.cpp b/Engine/Plugins/Media/BinkMedia/Source/BinkMediaPlayer/Private/BinkMovieStreamer.cpp
--- a/Engine/Plugins/Media/BinkMedia/Source/BinkMediaPlayer/Private/BinkMovieStreamer.cpp
+++ b/Engine/Plugins/Media/BinkMedia/Source/BinkMediaPlayer/Private/BinkMovieStreamer.cpp
@@ -235,22 +235,22 @@ bool FBinkMovieStreamer::OpenNextMovie()
 			bnk = BinkPluginOpen(TCHAR_TO_ANSI(*FileRootPath), soundTrack, soundTrackStart, bufferMode, FileOffset);
 		}
 #else
-		FString ExternalPath = PlatformFile.ConvertToAbsolutePathForExternalAppForRead(*FullMoviePath);
-		if (FPlatformFileManager::Get().GetPlatformFile().FileExists(*ExternalPath))
+		if (FPlatformFileManager::Get().GetPlatformFile().FileExists(*FullMoviePath)) 
 		{
-#if PLATFORM_WINDOWS
+			const FString ExternalPath = PlatformFile.ConvertToAbsolutePathForExternalAppForRead(*FullMoviePath);
+#if PLATFORM_WINDOWS 
 			bnk = BinkPluginOpenUTF16((unsigned short*)*ExternalPath, soundTrack, soundTrackStart, bufferMode, 0);
 #else
 			bnk = BinkPluginOpen(TCHAR_TO_ANSI(*ExternalPath), soundTrack, soundTrackStart, bufferMode, 0);
 #endif
 		}
-		if (!bnk)
+		if (!bnk) 
 		{
 			FString CookPath = *BinkUE4CookOnTheFlyPath(FPaths::ConvertRelativePathToFull(BINKMOVIEPATH), *MoviePathTbl[i]);
-			ExternalPath = PlatformFile.ConvertToAbsolutePathForExternalAppForRead(*FullMoviePath);
-			if (FPlatformFileManager::Get().GetPlatformFile().FileExists(*ExternalPath))
+			if (FPlatformFileManager::Get().GetPlatformFile().FileExists(*CookPath)) 
 			{
-#if PLATFORM_WINDOWS
+				const FString ExternalPath = PlatformFile.ConvertToAbsolutePathForExternalAppForRead(*CookPath);
+#if PLATFORM_WINDOWS 
 				bnk = BinkPluginOpenUTF16((unsigned short *)*ExternalPath, soundTrack, soundTrackStart, bufferMode, 0);
 #else
 				bnk = BinkPluginOpen(TCHAR_TO_ANSI(*ExternalPath), soundTrack, soundTrackStart, bufferMode, 0);
