Bink: Fix for GetRenderTargetTexture on the BackBuffer not returning valid backbuffer at runtime.

#Synced-CL 49716836

diff --git a/Engine/Plugins/Media/BinkMedia/Source/BinkMediaPlayer/Private/BinkFunctionLibrary.cpp b/Engine/Plugins/Media/BinkMedia/Source/BinkMediaPlayer/Private/BinkFunctionLibrary.cpp
--- a/Engine/Plugins/Media/BinkMedia/Source/BinkMediaPlayer/Private/BinkFunctionLibrary.cpp
+++ b/Engine/Plugins/Media/BinkMedia/Source/BinkMediaPlayer/Private/BinkFunctionLibrary.cpp
@@ -7,12 +7,13 @@
 #include "Engine/GameViewportClient.h"
 #include "Engine/Engine.h"
 #include "RenderingThread.h"
+#include "RHI.h"
 
 
 
 extern TSharedPtr<FBinkMovieStreamer, ESPMode::ThreadSafe> MovieStreamer;
 
-// Note: Has to be in a seperate function because you can't do #if inside a render command macro
+// Note: Has to be in a separate function because you can't do #if inside a render command macro
 static void Bink_DrawOverlays_Internal(FRHICommandListImmediate &RHICmdList, FTextureRHIRef BackBuffer, FVector2D ScreenSize) {
 	if(!BackBuffer.GetReference()) 
 	{
@@ -41,14 +42,21 @@ void UBinkFunctionLibrary::Bink_DrawOverlays()
 	GEngine->GameViewport->GetViewportSize(ScreenSize);
 	FTextureRHIRef BackBuffer = GEngine->GameViewport->Viewport->GetRenderTargetTexture();
 
-	if (!BackBuffer.IsValid()) {
-		return;
+	if (BackBuffer.IsValid()) {
+		ENQUEUE_RENDER_COMMAND(BinkOverlays)([ScreenSize, BackBuffer](FRHICommandListImmediate& RHICmdList) 
+		{ 
+			Bink_DrawOverlays_Internal(RHICmdList, BackBuffer, ScreenSize);
+		});
+	} else {
+		FViewportRHIRef ViewportRHI = GEngine->GameViewport->Viewport->GetViewportRHI();
+		if (ViewportRHI.IsValid()) {
+			ENQUEUE_RENDER_COMMAND(BinkOverlays)([ScreenSize, ViewportRHI](FRHICommandListImmediate& RHICmdList) 
+			{ 
+				FTextureRHIRef BackBuffer = RHIGetViewportBackBuffer(ViewportRHI);
+				Bink_DrawOverlays_Internal(RHICmdList, BackBuffer, ScreenSize);
+			});
+		}
 	}
-
-	ENQUEUE_RENDER_COMMAND(BinkOverlays)([ScreenSize, BackBuffer](FRHICommandListImmediate& RHICmdList) 
-	{ 
-		Bink_DrawOverlays_Internal(RHICmdList, BackBuffer, ScreenSize);
-	});
 }
 
 FTimespan UBinkFunctionLibrary::BinkLoadingMovie_GetDuration() 
