[bink] Fixed Metal backend issue with DontLoad_Store -> should be Clear_Store

#Synced-CL 49716836

diff --git a/Engine/Plugins/Media/BinkMedia/Source/BinkMediaPlayer/Private/binktexturesRHI.cpp b/Engine/Plugins/Media/BinkMedia/Source/BinkMediaPlayer/Private/binktexturesRHI.cpp
--- a/Engine/Plugins/Media/BinkMedia/Source/BinkMediaPlayer/Private/binktexturesRHI.cpp
+++ b/Engine/Plugins/Media/BinkMedia/Source/BinkMediaPlayer/Private/binktexturesRHI.cpp
@@ -482,7 +482,7 @@ static void Draw_textures(BINKTEXTURES* ptextures, BINKSHADERS* pshaders, void*
 
 
 			RHICmdList.Transition(FRHITransitionInfo(consts->RenderTargets[0].GetTexture()->GetRHI()->GetTexture2D(), ERHIAccess::Unknown, ERHIAccess::RTV));
-			FRHIRenderPassInfo RPInfo(consts->RenderTargets[0].GetTexture()->GetRHI()->GetTexture2D(), ERenderTargetActions::DontLoad_Store);
+			FRHIRenderPassInfo RPInfo(consts->RenderTargets[0].GetTexture()->GetRHI()->GetTexture2D(), ERenderTargetActions::Clear_Store);
 			RHICmdList.BeginRenderPass(RPInfo, TEXT("BinkPass"));
 
 			FGraphicsPipelineStateInitializer GraphicsPSOInit;
