Bink: update for new BinkHL interface

#Synced-CL 52144611

diff --git a/Engine/Plugins/Media/BinkMedia/Source/BinkMediaPlayer/BinkMediaPlayer.Build.cs b/Engine/Plugins/Media/BinkMedia/Source/BinkMediaPlayer/BinkMediaPlayer.Build.cs
--- a/Engine/Plugins/Media/BinkMedia/Source/BinkMediaPlayer/BinkMediaPlayer.Build.cs
+++ b/Engine/Plugins/Media/BinkMedia/Source/BinkMediaPlayer/BinkMediaPlayer.Build.cs
@@ -19,6 +19,8 @@ public class BinkMediaPlayer : ModuleRules
 		PrivateDependencyModuleNames.Add("RenderCore");
 		PrivateDependencyModuleNames.Add("RHI");
 		PrivateDependencyModuleNames.Add("Renderer");
+		PrivateDependencyModuleNames.Add("Slate");
+		PrivateDependencyModuleNames.Add("SlateCore");
 		
 		AddEngineThirdPartyPrivateStaticDependencies(Target, "BinkMediaPlayerSDK");
 		
diff --git a/Engine/Plugins/Media/BinkMedia/Source/BinkMediaPlayer/Private/Assets/BinkMediaPlayer.cpp b/Engine/Plugins/Media/BinkMedia/Source/BinkMediaPlayer/Private/Assets/BinkMediaPlayer.cpp
--- a/Engine/Plugins/Media/BinkMedia/Source/BinkMediaPlayer/Private/Assets/BinkMediaPlayer.cpp
+++ b/Engine/Plugins/Media/BinkMedia/Source/BinkMediaPlayer/Private/Assets/BinkMediaPlayer.cpp
@@ -10,19 +10,15 @@
 #include "Engine/Texture.h"
 #include "Misc/Paths.h"
 #include "RenderingThread.h"
-#include "HDRHelper.h"
-
 
 #include "HAL/PlatformFileManager.h"
 #include "Internationalization/Culture.h"
+#include "Internationalization/Internationalization.h"
 #include "SubtitleManager.h"
 #include "UObject/UObjectIterator.h"
-#include "binkplugin_ue4.h"
-#include "egttypes.h"
+#include "BinkUE.h"
 #include "RectLightTexture.h"
 
-//#include "Media/Public/IMediaEventSink.h" //For EMediaEvent
-
 static void Command_ShowBinks()
 {
 	for (TObjectIterator<UBinkMediaPlayer> It; It; ++It)
@@ -47,7 +43,7 @@ static void Command_ShowBinks()
 			{
 				DrawStyleStr = TEXT("Overlay");
 			}
-			
+
 			const double DurationSeconds = It->GetDuration().GetTotalSeconds();
 			const double PlaybackCompletion = (DurationSeconds > 0.0) ? It->GetTime().GetTotalSeconds() / DurationSeconds : 0.0;
 
@@ -63,6 +59,7 @@ static FAutoConsoleCommand ConsoleCmdShowBinks(
 
 UBinkMediaPlayer::UBinkMediaPlayer( const FObjectInitializer& ObjectInitializer )
 	: Super(ObjectInitializer)
+	, FTickableGameObject(ETickableTickType::Never)
 	, Looping(true)
 	, StartImmediately(true)
 	, DelayedOpen(true)
@@ -80,6 +77,7 @@ UBinkMediaPlayer::UBinkMediaPlayer( const FObjectInitializer& ObjectInitializer
 	, CurrentDrawStyle()
 	, CurrentLayerDepth()
     , bnk()
+    , bink_textures()
     , paused()
     , reached_end()
 {
@@ -102,9 +100,9 @@ FTimespan UBinkMediaPlayer::GetDuration() const
 	double ms = 0;
 	if(bnk) 
 	{
-		BINKPLUGININFO bpinfo = {};
-		BinkPluginInfo(bnk, &bpinfo);
-		ms = ((double)bpinfo.Frames) * ((double)bpinfo.FrameRateDiv) * 1000.0 / ((double)bpinfo.FrameRate);
+		BINKHLINFO info = {};
+		BinkHLGetInfo(bnk, &info);
+		ms = ((double)info.Frames) * ((double)info.FrameRateDiv) * 1000.0 / ((double)info.FrameRate);
 	}
 	return FTimespan::FromMilliseconds(ms);
 }
@@ -113,9 +111,9 @@ float UBinkMediaPlayer::GetRate() const
 {
 	if(bnk) 
 	{
-		BINKPLUGININFO bpinfo = {};
-		BinkPluginInfo(bnk, &bpinfo);
-		return bpinfo.PlaybackState == 0 && !paused ? 1 : 0;
+		BINKHLINFO info = {};
+		BinkHLGetInfo(bnk, &info);
+		return info.PlaybackState == 0 && !paused ? 1 : 0;
 	}
 	return 0; 
 }
@@ -125,9 +123,9 @@ FTimespan UBinkMediaPlayer::GetTime() const
 	double ms = 0;
 	if(bnk) 
 	{
-		BINKPLUGININFO bpinfo = {};
-		BinkPluginInfo(bnk, &bpinfo);
-		ms = ((double)bpinfo.FrameNum) * ((double)bpinfo.FrameRateDiv) *1000.0 / ((double)bpinfo.FrameRate);
+		BINKHLINFO info = {};
+		BinkHLGetInfo(bnk, &info);
+		ms = ((double)info.FrameNum) * ((double)info.FrameRateDiv) * 1000.0 / ((double)info.FrameRate);
 	}
 	return FTimespan::FromMilliseconds(ms);
 }
@@ -136,9 +134,9 @@ bool UBinkMediaPlayer::IsLooping() const
 { 
 	if(bnk) 
 	{
-		BINKPLUGININFO bpinfo = {};
-		BinkPluginInfo(bnk, &bpinfo);
-		if (bpinfo.PlaybackState < 3 && !bpinfo.LoopsRemaining) 
+		BINKHLINFO info = {};
+		BinkHLGetInfo(bnk, &info);
+		if (info.PlaybackState < 3 && !info.LoopsRemaining) 
 		{
 			return true;
 		}
@@ -150,31 +148,31 @@ bool UBinkMediaPlayer::IsPaused() const
 { 
 	if(bnk) 
 	{
-		BINKPLUGININFO bpinfo = {};
-		BinkPluginInfo(bnk, &bpinfo);
-		return bpinfo.PlaybackState == 1 || paused; // TODO: When the video is paused, the PlaybackState is 0 (should be 1).
+		BINKHLINFO info = {};
+		BinkHLGetInfo(bnk, &info);
+		return info.PlaybackState == 1 || paused;
 	}
 	return false;
 }
 
 bool UBinkMediaPlayer::IsPlaying() const 
 {
-    if(bnk) 
+	if(bnk) 
 	{
-        BINKPLUGININFO bpinfo = {};
-        BinkPluginInfo(bnk, &bpinfo);
-        return bpinfo.PlaybackState == 0 && !paused;
-    }
-    return false;
+		BINKHLINFO info = {};
+		BinkHLGetInfo(bnk, &info);
+		return info.PlaybackState == 0 && !paused;
+	}
+	return false;
 }
 
 bool UBinkMediaPlayer::IsGotoing() const
 {
 	if (bnk)
 	{
-		BINKPLUGININFO bpinfo = {};
-		BinkPluginInfo(bnk, &bpinfo);
-		return bpinfo.PlaybackState == 2;
+		BINKHLINFO info = {};
+		BinkHLGetInfo(bnk, &info);
+		return info.PlaybackState == 2;
 	}
 	return false;
 }
@@ -201,7 +199,8 @@ bool UBinkMediaPlayer::SetLooping( bool InLooping )
 {
 	if(bnk) 
 	{
-		BinkPluginLoop(bnk, InLooping ? 0 : 1);
+		// 0 = loop forever, 1 = play once
+		BinkHLLoopCount(bnk, InLooping ? 0 : 1);
 	}
 	return false;
 }
@@ -210,30 +209,28 @@ bool UBinkMediaPlayer::SetRate( float Rate )
 { 
 	if(bnk) 
 	{
-		BINKPLUGININFO bpinfo = {};
-		BinkPluginInfo(bnk, &bpinfo);
-		if((bpinfo.PlaybackState == 1 || paused) && Rate == 1) // If paused and set the rate to 1.0
+		BINKHLINFO info = {};
+		BinkHLGetInfo(bnk, &info);
+		if ((info.PlaybackState == 1 || paused) && Rate == 1) 
 		{
-			BinkPluginPause(bnk, 0);
-			paused = false;
+			BinkHLPause(bnk, 0);
+			paused      = false;
 			reached_end = false;
-			//MediaEvent.Broadcast(EMediaEvent::PlaybackResumed);
 			return true;
 		} 
-		else if(bpinfo.PlaybackState == 3 && Rate == 1) 
+		else if (info.PlaybackState == 3 && Rate == 1) 
 		{
-			BinkPluginGoto(bnk, 1, -1);
-			BinkPluginPause(bnk, 0);
-			paused = false;
+			// Rewind to start after reaching end
+			BinkHLGoto(bnk, 1, 4);
+			BinkHLPause(bnk, 0);
+			paused      = false;
 			reached_end = false;
-			//MediaEvent.Broadcast(EMediaEvent::PlaybackResumed);
 			return true;
 		} 
-		else if (bpinfo.PlaybackState == 0 && Rate == 0) 
+		else if (info.PlaybackState == 0 && Rate == 0) 
 		{
-			BinkPluginPause(bnk, -1);
+			BinkHLPause(bnk, 1);
 			paused = true;
-			//MediaEvent.Broadcast(EMediaEvent::PlaybackSuspended);
 			OnPlaybackSuspended.Broadcast();
 			return true;
 		}
@@ -245,9 +242,8 @@ void UBinkMediaPlayer::SetVolume( float Volume )
 { 
 	if(bnk) 
 	{
-		// Clamp 0 .. 1
 		Volume = Volume < 0 ? 0 : Volume > 1 ? 1 : Volume;
-		BinkPluginVolume(bnk, Volume);
+		BinkHLVolume(bnk, Volume);
 	}
 }
 
@@ -257,12 +253,12 @@ bool UBinkMediaPlayer::Seek( const FTimespan& InTime )
 	{
 		return false;
 	}
-	BINKPLUGININFO bpinfo = {};
-	BinkPluginInfo(bnk, &bpinfo);
-	U32 desiredFrame = (U32)(InTime.GetTotalMilliseconds() * ((double)bpinfo.FrameRate) / (1000.0*((double)bpinfo.FrameRateDiv)));
-	if (bpinfo.FrameNum != (desiredFrame + 1))
+	BINKHLINFO info = {};
+	BinkHLGetInfo(bnk, &info);
+	U32 desiredFrame = (U32)(InTime.GetTotalMilliseconds() * ((double)info.FrameRate) / (1000.0 * ((double)info.FrameRateDiv)));
+	if (info.FrameNum != (desiredFrame + 1))
 	{
-		BinkPluginGoto(bnk, desiredFrame + 1, -1);
+		BinkHLGoto(bnk, desiredFrame + 1, 4);
 		reached_end = false;
 	}
 	return true;
@@ -284,13 +280,11 @@ void UBinkMediaPlayer::PostLoad()
 }
 
 #if BINKPLUGIN_UE4_EDITOR
-
 void UBinkMediaPlayer::PostEditChangeProperty( FPropertyChangedEvent& PropertyChangedEvent ) 
 {
 	Super::PostEditChangeProperty(PropertyChangedEvent);
 	InitializePlayer();
 }
-
 #endif
 
 bool UBinkMediaPlayer::Open(const FString& Url) 
@@ -315,101 +309,87 @@ bool UBinkMediaPlayer::Open(const FString& Url)
 
 	if(bnk) 
 	{
-		ENQUEUE_RENDER_COMMAND(BinkMediaPlayer_Open_CloseBink)([bnk=bnk](FRHICommandListImmediate& RHICmdList) 
-		{ 
-			BinkPluginClose(bnk);
-		});
-		bnk = NULL;
+		BinkCloseOnRenderThread(bnk, bink_textures);
 	}
 
-	BinkPluginLimitSpeakers(GetNumSpeakers());
-
-	// Use the platform file layer to open the media file. We
-	// need to access Bink specific information to allow
-	// for playing media that is embedded in the APK, OBBs,
-	// and/or PAKs.
-
 	IPlatformFile& PlatformFile = FPlatformFileManager::Get().GetPlatformFile();
 
 #if PLATFORM_ANDROID
-	// Construct a canonical path for the movie.
 	FString MoviePath = Url;
 	FPaths::NormalizeFilename(MoviePath);
 
-	// Don't bother trying to play it if we can't find it.
 	if (!IAndroidPlatformFile::GetPlatformPhysical().FileExists(*MoviePath)) {
 		UE_LOG(LogBinkMoviePlayer, Error, TEXT("UBinkMediaPlayer::Open: Failed! File doesn't exist. URL :%s File :%s"), *Url, *MoviePath);
 		return false;
 	}
 
-	// Get information about the movie.
-	int64 FileOffset = IAndroidPlatformFile::GetPlatformPhysical().FileStartOffset(*MoviePath);
-	FString FileRootPath = IAndroidPlatformFile::GetPlatformPhysical().FileRootPath(*MoviePath);
+	int64 FileOffset      = IAndroidPlatformFile::GetPlatformPhysical().FileStartOffset(*MoviePath);
+	FString FileRootPath  = IAndroidPlatformFile::GetPlatformPhysical().FileRootPath(*MoviePath);
 
-	// Play the movie as a file or asset.
 	if (IAndroidPlatformFile::GetPlatformPhysical().IsAsset(*MoviePath)) 
 	{
-		if (JNIEnv* env = FAndroidApplication::GetJavaEnv()) 
+		FString APKPath = BinkGetAndroidAPKPath();
+		if (!APKPath.IsEmpty())
 		{
-			extern struct android_app* GNativeAndroidApp;
-			jclass clazz = env->GetObjectClass(GNativeAndroidApp->activity->clazz);
-			jmethodID methodID = env->GetMethodID(clazz, "getPackageCodePath", "()Ljava/lang/String;");
-			jobject result = env->CallObjectMethod(GNativeAndroidApp->activity->clazz, methodID);
-			jboolean isCopy;
-			const char *apkPath = env->GetStringUTFChars((jstring)result, &isCopy);
-			bnk = BinkPluginOpen(apkPath, BinkSoundTrack, BinkSoundTrackStart, BinkBufferMode, FileOffset);
-			env->ReleaseStringUTFChars((jstring)result, apkPath);
+			bnk = BinkHLOpen(TCHAR_TO_UTF8(*APKPath), (U32)BinkSoundTrack, (U32)BinkSoundTrackStart, (U32)BinkBufferMode, (U64)FileOffset);
 		}
 	}
 	else 
 	{
-		bnk = BinkPluginOpen(TCHAR_TO_ANSI(*FileRootPath), BinkSoundTrack, BinkSoundTrackStart, BinkBufferMode, FileOffset);
+		bnk = BinkHLOpen(TCHAR_TO_UTF8(*FileRootPath), (U32)BinkSoundTrack, (U32)BinkSoundTrackStart, (U32)BinkBufferMode, (U64)FileOffset);
 	}
 #else
 	if (FPlatformFileManager::Get().GetPlatformFile().FileExists(*Url)) 
 	{
-#if PLATFORM_WINDOWS 
-		bnk = BinkPluginOpenUTF16((unsigned short *)*PlatformFile.ConvertToAbsolutePathForExternalAppForRead(*Url), BinkSoundTrack, BinkSoundTrackStart, BinkBufferMode, 0);
-#else
-		bnk = BinkPluginOpen(TCHAR_TO_ANSI(*PlatformFile.ConvertToAbsolutePathForExternalAppForRead(*Url)), BinkSoundTrack, BinkSoundTrackStart, BinkBufferMode, 0);
-#endif
+		FString AbsPath = PlatformFile.ConvertToAbsolutePathForExternalAppForRead(*Url);
+		bnk = BinkHLOpen(TCHAR_TO_UTF8(*AbsPath), (U32)BinkSoundTrack, (U32)BinkSoundTrackStart, (U32)BinkBufferMode, 0);
 	}
 #endif
+
 	if(!bnk) 
 	{
-		//MediaEvent.Broadcast(EMediaEvent::MediaOpenFailed);
-		UE_LOG(LogBinkMoviePlayer, Error, TEXT("UBinkMediaPlayer::Open: Failed! BinkPluginOpen failed. Invalid bnk."));
+		UE_LOG(LogBinkMoviePlayer, Error, TEXT("UBinkMediaPlayer::Open: Failed! BinkHLOpen failed. URL: %s Error: %s"), *Url, UTF8_TO_TCHAR(BinkGetError()));
 		return false;
 	}
 
-	// Try to open subtitles...
+	// Create CPU-side decode buffers immediately on the game thread.
+	// GPU textures are created lazily on the render thread on first Draw_textures call.
+	bink_textures = FBinkTextures::CreateCPU(bnk, nullptr);
+
+	if (!bink_textures)
 	{
-		// Determine what out current culture is, and grab the most appropriate set of subtitles for it
-		FInternationalization& Internationalization = FInternationalization::Get();
+		UE_LOG(LogBinkMoviePlayer, Error, TEXT("UBinkMediaPlayer::Open: Failed to create CPU decode buffers. URL: %s"), *Url);
+		BinkCloseOnRenderThread(bnk, bink_textures);
+		return false;
+	}
+
+	paused      = false;
+	reached_end = false;
+	CurrentHasSubtitles = 0;
 
-		const TArray<FString> PrioritizedLanguageNames = Internationalization.GetPrioritizedCultureNames(Internationalization.GetCurrentLanguage()->GetName());
+	// Search for a language-matched .srt file and load it via BinkLoadSubtitles.
+	{
+		IPlatformFile& PF = FPlatformFileManager::Get().GetPlatformFile();
+		FInternationalization& I18N = FInternationalization::Get();
+		const TArray<FString> LangNames = I18N.GetPrioritizedCultureNames(I18N.GetCurrentLanguage()->GetName());
 
-		CurrentHasSubtitles = 0;
-		for (const FString& LanguageName : PrioritizedLanguageNames)
+		for (const FString& LangName : LangNames)
 		{
-			int32 Pos = INDEX_NONE;
-			if (Url.FindLastChar(TEXT('.'), Pos))
+			int32 DotPos = INDEX_NONE;
+			if (Url.FindLastChar(TEXT('.'), DotPos))
 			{
-				const int32 PathEndPos = Url.FindLastCharByPredicate([](TCHAR C) { return C == TEXT('/') || C == TEXT('\\'); });
-				if (PathEndPos != INDEX_NONE && PathEndPos > Pos)
+				const int32 SlashPos = Url.FindLastCharByPredicate([](TCHAR C){ return C == TEXT('/') || C == TEXT('\\'); });
+				if (SlashPos != INDEX_NONE && SlashPos > DotPos)
 				{
-					// The dot found was part of the path rather than the name
-					Pos = INDEX_NONE;
+					DotPos = INDEX_NONE;
 				}
 			}
-			FString SubtitleUrl = (Pos == INDEX_NONE ? Url : Url.Left(Pos)) + "_" + LanguageName + ".srt";
-			if (FPlatformFileManager::Get().GetPlatformFile().FileExists(*SubtitleUrl))
+			FString SubUrl = (DotPos == INDEX_NONE ? Url : Url.Left(DotPos)) + TEXT("_") + LangName + TEXT(".srt");
+
+			if (PF.FileExists(*SubUrl))
 			{
-#if PLATFORM_WINDOWS 
-				if (BinkPluginLoadSubtitlesUTF16(bnk, (unsigned short*)*PlatformFile.ConvertToAbsolutePathForExternalAppForRead(*SubtitleUrl)))
-#else
-				if (BinkPluginLoadSubtitles(bnk, TCHAR_TO_ANSI(*PlatformFile.ConvertToAbsolutePathForExternalAppForRead(*SubtitleUrl))))
-#endif
+				FString AbsSub = PF.ConvertToAbsolutePathForExternalAppForRead(*SubUrl);
+				if (BinkLoadSubtitles(bnk, TCHAR_TO_UTF8(*AbsSub)))
 				{
 					CurrentHasSubtitles = 1;
 					break;
@@ -417,31 +397,25 @@ bool UBinkMediaPlayer::Open(const FString& Url)
 			}
 		}
 	}
-    
-	paused = false;
-	reached_end = false;
+
+	SetTickableTickType(ETickableTickType::Conditional);
 	HandleMediaPlayerMediaOpened(Url);
 	return true;
 }
 
 void UBinkMediaPlayer::Close() 
 {
+	SetTickableTickType(ETickableTickType::Never);
 	if(bnk) 
 	{
-		if (CurrentHasSubtitles) 
+		if (CurrentHasSubtitles)
 		{
 			CurrentHasSubtitles = 0;
-
-			// Clear the movie subtitle for this object
-			TArray<FString> StopSubtitles = { TEXT("") };
-			FSubtitleManager::GetSubtitleManager()->SetMovieSubtitle(this, StopSubtitles);
+			TArray<FString> Empty = { TEXT("") };
+			FSubtitleManager::GetSubtitleManager()->SetMovieSubtitle(this, Empty);
 		}
 
-		ENQUEUE_RENDER_COMMAND(BinkMediaPlayer_Close_CloseBink)([bnk=bnk](FRHICommandListImmediate& RHICmdList) 
-		{ 
-			BinkPluginClose(bnk);
-		});
-		bnk = NULL;
+		BinkCloseOnRenderThread(bnk, bink_textures);
 	}
 	CurrentUrl = FString();
 	HandleMediaPlayerMediaClosed();
@@ -463,10 +437,8 @@ void UBinkMediaPlayer::InitializePlayer()
 			return;
 		}
 
-		// open the new media file
 		bool OpenedSuccessfully = false;
 
-		IPlatformFile& PlatformFile = FPlatformFileManager::Get().GetPlatformFile();
 		FString FullUrl = FPaths::ConvertRelativePathToFull(FPaths::IsRelative(URL) ? BINKCONTENTPATH / URL : URL);
 		OpenedSuccessfully = Open(FullUrl);
 		if (!OpenedSuccessfully) 
@@ -475,15 +447,14 @@ void UBinkMediaPlayer::InitializePlayer()
 			OpenedSuccessfully = Open(cookpath);
 		}
 
-		// finish initialization
 		if (OpenedSuccessfully) 
 		{
-			CurrentUrl = URL;
-			CurrentBinkBufferMode = BinkBufferMode;
-			CurrentBinkSoundTrack = BinkSoundTrack;
+			CurrentUrl              = URL;
+			CurrentBinkBufferMode   = BinkBufferMode;
+			CurrentBinkSoundTrack   = BinkSoundTrack;
 			CurrentBinkSoundTrackStart = BinkSoundTrackStart;
-			CurrentDrawStyle = BinkDrawStyle;
-			CurrentLayerDepth = BinkLayerDepth;
+			CurrentDrawStyle        = BinkDrawStyle;
+			CurrentLayerDepth       = BinkLayerDepth;
 		}
 	}
 
@@ -495,26 +466,37 @@ void UBinkMediaPlayer::HandleMediaPlayerMediaClosed()
 {
 	MediaChangedEvent.Broadcast();
 	OnMediaClosed.Broadcast();
-	//MediaEvent.Broadcast(EMediaEvent::MediaClosed);
 }
 
 void UBinkMediaPlayer::HandleMediaPlayerMediaOpened( FString OpenedUrl ) 
 {
 	MediaChangedEvent.Broadcast();
 	OnMediaOpened.Broadcast(OpenedUrl);
-	//MediaEvent.Broadcast(EMediaEvent::MediaOpened);
 }
 
 void UBinkMediaPlayer::Tick(float DeltaTime) 
 {
-	// Check for if we should issue the reached end event
+	// Fire the reached-end event once when playback stops at the end
 	if (bnk && !reached_end && !IsPlaying() && !IsGotoing() && !IsPaused() && !IsLooping()) 
 	{
 		OnMediaReachedEnd.Broadcast();
 		reached_end = true;
-		//MediaEvent.Broadcast(EMediaEvent::PlaybackEndReached);
 	}
-	if (bnk && GEngine && GEngine->GameViewport) 
+
+	// Update subtitles each tick using the native BinkCurrentSubtitle iterator.
+	if (bnk && CurrentHasSubtitles)
+	{
+		TArray<FString> SubtitlesText;
+		U32 iterate = 0;
+		while (const char* sub = BinkCurrentSubtitle(bnk, &iterate, nullptr, nullptr))
+		{
+			SubtitlesText.Add(FString(UTF8_TO_TCHAR(sub)));
+		}
+		FSubtitleManager::GetSubtitleManager()->SetMovieSubtitle(this, SubtitlesText);
+	}
+
+	// Overlay draw style: schedule a draw to the backbuffer each frame
+	if (bnk && BinkDrawStyle != 0 && GEngine && GEngine->GameViewport) 
 	{
 		if (!IsPlaying() && !IsPaused()) 
 		{
@@ -523,110 +505,50 @@ void UBinkMediaPlayer::Tick(float DeltaTime)
 
 		FVector2D screenSize;
 		GEngine->GameViewport->GetViewportSize(screenSize);
-		if (CurrentHasSubtitles == 1)
-		{
-			TArray<FString> SubtitlesText;
-			unsigned i = 0;
-			while(const char *sub = BinkPluginCurrentSubtitle(bnk, &i))
-			{
-				SubtitlesText.Add((UTF8CHAR*)sub);
-			}
 
-			FSubtitleManager::GetSubtitleManager()->SetMovieSubtitle(this, SubtitlesText);
-		}
+		float ulx = BinkDestinationUpperLeft.X;
+		float uly = BinkDestinationUpperLeft.Y;
+		float lrx = BinkDestinationLowerRight.X;
+		float lry = BinkDestinationLowerRight.Y;
+
+		BINKHLINFO info = {};
+		BinkHLGetInfo(bnk, &info);
+		int binkw = (int)info.Width;
+		int binkh = (int)info.Height;
 
-		if (BinkDrawStyle != 0)
+		if (BinkDrawStyle == 1 /*OverlayFillScreenWithAspectRatio*/) 
 		{
-			BINKPLUGININFO bpinfo = {};
-			BinkPluginInfo(bnk, &bpinfo);
-			int binkw = bpinfo.Width;
-			int binkh = bpinfo.Height;
-
-			float ulx = BinkDestinationUpperLeft.X;
-			float uly = BinkDestinationUpperLeft.Y;
-			float lrx = BinkDestinationLowerRight.X;
-			float lry = BinkDestinationLowerRight.Y;
-
-			// figure out the x,y screencoords for all of the overlay types
-			if (BinkDrawStyle == 1/*BNK_DS_OverlayFillScreenWithAspectRatio*/) 
+			lrx = binkw / screenSize.X;
+			lry = binkh / screenSize.Y;
+			if (lrx > lry) 
 			{
-				lrx = binkw / screenSize.X;
-				lry = binkh / screenSize.Y;
-
-				if (lrx > lry) 
-				{
-					lry /= lrx;
-					lrx = 1;
-				}
-				else 
-				{
-					lrx /= lry;
-					lry = 1;
-				}
-				ulx = (1.0f - lrx) / 2.0f;
-				uly = (1.0f - lry) / 2.0f;
-				lrx += ulx;
-				lry += uly;
+				lry /= lrx; lrx = 1;
 			}
-			else if (BinkDrawStyle == 2/*BNK_DS_OverlayOriginalMovieSize*/) 
+			else 
 			{
-				ulx = (screenSize.X - binkw) / (2.0f * screenSize.X);
-				uly = (screenSize.Y - binkh) / (2.0f * screenSize.Y);
-				lrx = binkw / screenSize.X + ulx;
-				lry = binkh / screenSize.Y + uly;
+				lrx /= lry; lry = 1;
 			}
+			ulx = (1.0f - lrx) / 2.0f;
+			uly = (1.0f - lry) / 2.0f;
+			lrx += ulx;
+			lry += uly;
+		}
+		else if (BinkDrawStyle == 2 /*OverlayOriginalMovieSize*/) 
+		{
+			ulx = (screenSize.X - binkw) / (2.0f * screenSize.X);
+			uly = (screenSize.Y - binkh) / (2.0f * screenSize.Y);
+			lrx = binkw / screenSize.X + ulx;
+			lry = binkh / screenSize.Y + uly;
+		}
 
 #if PLATFORM_ANDROID
-			uly = 1 - uly;
-			lry = 1 - lry;
+		uly = 1 - uly;
+		lry = 1 - lry;
 #endif
 
-			ENQUEUE_RENDER_COMMAND(BinkScheduleOverlay)([bnk=bnk,ulx,uly,lrx,lry](FRHICommandListImmediate& RHICmdList) 
-			{ 
-				static const auto CVarHDROutputEnabled = IConsoleManager::Get().FindTConsoleVariableDataInt(TEXT("r.HDR.EnableHDROutput"));
-				static const auto CVarDisplayOutputDevice = IConsoleManager::Get().FindTConsoleVariableDataInt(TEXT("r.HDR.Display.OutputDevice"));
-				if (GRHISupportsHDROutput && CVarHDROutputEnabled->GetValueOnRenderThread() != 0)
-				{
-					EDisplayOutputFormat outDev = static_cast<EDisplayOutputFormat>(CVarDisplayOutputDevice->GetValueOnRenderThread());
-					float DisplayMaxLuminance = HDRGetDisplayMaximumLuminance();
-					switch (outDev)
-					{
-					// LDR
-					case EDisplayOutputFormat::SDR_sRGB:
-					case EDisplayOutputFormat::SDR_Rec709:
-					case EDisplayOutputFormat::SDR_ExplicitGammaMapping:
-						BinkPluginSetHdrSettings(bnk, 1, 1.0f, 80);
-						break;
-					// 1k nits
-					case EDisplayOutputFormat::HDR_ACES_1000nit_ST2084:
-						BinkPluginSetHdrSettings(bnk, 2, 1.0f, DisplayMaxLuminance);
-						break;
-					case EDisplayOutputFormat::HDR_ACES_1000nit_ScRGB:
-						BinkPluginSetHdrSettings(bnk, 1, 1.0f, DisplayMaxLuminance);
-						break;
-					// 2k nits
-					case EDisplayOutputFormat::HDR_ACES_2000nit_ST2084:
-						BinkPluginSetHdrSettings(bnk, 2, 1.0f, DisplayMaxLuminance);
-						break;
-					case EDisplayOutputFormat::HDR_ACES_2000nit_ScRGB:
-						BinkPluginSetHdrSettings(bnk, 1, 1.0f, DisplayMaxLuminance);
-						break;
-					// no tonemap
-					default:
-						BinkPluginSetHdrSettings(bnk, 0, 1.0f, 1000);
-						break;
-					}
-					BinkPluginSetRenderTargetFormat(bnk, 1);
-				}
-				else
-				{
-					BinkPluginSetHdrSettings(bnk, 1, 1.0f, 80);
-					BinkPluginSetRenderTargetFormat(bnk, 0);
-				}
-				BinkPluginSetDrawFlags(bnk, 0);
-				BinkPluginScheduleOverlay(bnk, ulx, uly, lrx, lry, 0);
-			});
-		}
+		// Bink_DrawOverlays() reads bink_textures at DrawBinks time (after all Ticks),
+		// so it reflects any Close() that happened this Tick.
+		PendingBinkOverlays.Add({ this, ulx, uly, lrx, lry });
 	}
 }
 
@@ -634,21 +556,28 @@ void UBinkMediaPlayer::UpdateTexture(const UTexture *tex, FRHICommandListImmedia
 {
 	check(IsInRenderingThread());
 
-	if (bnk && (BinkDrawStyle == 0 || isEditor)) 
+	if (!bnk || !bink_textures || (!isEditor && BinkDrawStyle != 0))
 	{
-		BinkPluginSetHdrSettings(bnk, tonemap, 1.0f, output_nits);
-		BinkPluginSetAlphaSettings(bnk, alpha);
-		BinkPluginSetDrawFlags(bnk, srgb_decode ? BinkDrawDecodeSRGB : 0);
-		BinkPluginSetRenderTargetFormat(bnk, is_hdr ? 1 : 0);
-		BinkPluginScheduleToTexture(bnk, BinkDestinationUpperLeft.X, BinkDestinationUpperLeft.Y, BinkDestinationLowerRight.X, BinkDestinationLowerRight.Y, 0, ref.GetReference(), width, height);
-		BinkActiveTextureRefs.Push(ref);
+		return;
+	}
 
-		// Update rect-lights with this texture
-		if (tex)
-		{
-			// Lock/Enqueue rect atlas refresh if that texture is used by a rect. light
-			RectLightAtlas::FAtlasTextureInvalidationScope InvalidationScope(tex);
-		}
+	// Draw directly. BinkHLProcess() has already run on this frame's render thread
+	// (from the DrawBinks render command which executes before UpdateDeferredResource).
+	S32 draw_flags = srgb_decode ? (S32)0x80000000 : 0;
+
+	bink_textures->SetRenderTarget(ref.GetReference(), ERenderTargetLoadAction::EClear);
+
+	bink_textures->SetHDRSettings(tonemap, 1.0f, output_nits);
+	bink_textures->SetAlphaSettings(alpha, draw_flags);
+	bink_textures->SetDrawPosition(
+		(float)BinkDestinationUpperLeft.X,  (float)BinkDestinationUpperLeft.Y,
+		(float)BinkDestinationLowerRight.X, (float)BinkDestinationLowerRight.Y);
+
+	bink_textures->Draw();
+
+	if (tex)
+	{
+		RectLightAtlas::FAtlasTextureInvalidationScope InvalidationScope(tex);
 	}
 }
 
@@ -663,7 +592,7 @@ void UBinkMediaPlayer::Draw(UTexture *texture, bool tonemap, int out_nits, float
 	{
 		return;
 	}
-	int width = texture->GetSurfaceWidth();
+	int width  = texture->GetSurfaceWidth();
 	int height = texture->GetSurfaceHeight();
 	void *native = ref->GetNativeResource();
 	if (!native) 
@@ -674,17 +603,18 @@ void UBinkMediaPlayer::Draw(UTexture *texture, bool tonemap, int out_nits, float
 
 	struct parms_t 
 	{
-		UBinkMediaPlayer *player;
-		UTexture* tex;
-		FTextureRHIRef ref;
-		void *native;
-		int width, height;
-		bool tonemap;
-		bool srgb_decode;
-		bool hdr;
-		int out_nits;
-		float alpha;
+		UBinkMediaPlayer* player;
+		UTexture*         tex;
+		FTextureRHIRef    ref;
+		void*             native;
+		int               width, height;
+		bool              tonemap;
+		bool              srgb_decode;
+		bool              hdr;
+		int               out_nits;
+		float             alpha;
 	} parms = { this, texture, ref, native, width, height, tonemap, srgb_decode, hdr, out_nits, alpha };
+
 	ENQUEUE_RENDER_COMMAND(BinkMediaPlayer_Draw)([parms](FRHICommandListImmediate& RHICmdList) 
 	{ 
 		parms.player->UpdateTexture(parms.tex, RHICmdList, parms.ref, parms.native, parms.width, parms.height, false, parms.tonemap, parms.out_nits, parms.alpha, parms.srgb_decode, parms.hdr);
@@ -695,9 +625,9 @@ FIntPoint UBinkMediaPlayer::GetDimensions() const
 {
 	if (bnk)
 	{
-		BINKPLUGININFO bpinfo = {};
-		BinkPluginInfo(bnk, &bpinfo);
-		return FIntPoint(bpinfo.Width, bpinfo.Height);
+		BINKHLINFO info = {};
+		BinkHLGetInfo(bnk, &info);
+		return FIntPoint(info.Width, info.Height);
 	}
 	return FIntPoint(0, 0);
 }
@@ -706,9 +636,9 @@ float UBinkMediaPlayer::GetFrameRate() const
 {
 	if (bnk)
 	{
-		BINKPLUGININFO bpinfo = {};
-		BinkPluginInfo(bnk, &bpinfo);
-		return (float)(((double)bpinfo.FrameRate) / ((double)bpinfo.FrameRateDiv));
+		BINKHLINFO info = {};
+		BinkHLGetInfo(bnk, &info);
+		return (float)(((double)info.FrameRate) / ((double)info.FrameRateDiv));
 	}
 	return 0;
 }
diff --git a/Engine/Plugins/Media/BinkMedia/Source/BinkMediaPlayer/Private/Assets/BinkMediaTexture.cpp b/Engine/Plugins/Media/BinkMedia/Source/BinkMediaPlayer/Private/Assets/BinkMediaTexture.cpp
--- a/Engine/Plugins/Media/BinkMedia/Source/BinkMediaPlayer/Private/Assets/BinkMediaTexture.cpp
+++ b/Engine/Plugins/Media/BinkMedia/Source/BinkMediaPlayer/Private/Assets/BinkMediaTexture.cpp
@@ -19,15 +19,15 @@ UBinkMediaTexture::UBinkMediaTexture( const FObjectInitializer& ObjectInitialize
 	, ReleasePlayerFence()
 {
 	NeverStream = true;
-	if (bink_gpu_api_hdr) 
+	if (GRHISupportsHDROutput)
 	{
 		PixelFormat = PF_A2B10G10R10;
-		DecodeSRGB = true;
-	} 
-	else 
+		DecodeSRGB  = true;
+	}
+	else
 	{
 		PixelFormat = PF_B8G8R8A8;
-		DecodeSRGB = true;
+		DecodeSRGB  = true;
 	}
 }
 
diff --git a/Engine/Plugins/Media/BinkMedia/Source/BinkMediaPlayer/Private/BinkFunctionLibrary.cpp b/Engine/Plugins/Media/BinkMedia/Source/BinkMediaPlayer/Private/BinkFunctionLibrary.cpp
--- a/Engine/Plugins/Media/BinkMedia/Source/BinkMediaPlayer/Private/BinkFunctionLibrary.cpp
+++ b/Engine/Plugins/Media/BinkMedia/Source/BinkMediaPlayer/Private/BinkFunctionLibrary.cpp
@@ -3,58 +3,136 @@
 
 #include "BinkFunctionLibrary.h"
 
+#include "BinkMediaPlayer.h"
 #include "BinkMovieStreamer.h"
+#include "BinkMediaPlayerPrivate.h"
 #include "Engine/GameViewportClient.h"
 #include "Engine/Engine.h"
 #include "RenderingThread.h"
 #include "RHI.h"
-
-
+#include "Framework/Application/SlateApplication.h"
+#include "Rendering/SlateRenderer.h"
+#include "Slate/SlateViewportProvider.h"
 
 extern TSharedPtr<FBinkMovieStreamer, ESPMode::ThreadSafe> MovieStreamer;
 
-// Note: Has to be in a separate function because you can't do #if inside a render command macro
-static void Bink_DrawOverlays_Internal(FRHICommandListImmediate &RHICmdList, FRHITexture* BackBuffer, FVector2D ScreenSize)
+// Render-thread draw of all overlay Binks to the backbuffer.
+// Each job has already had its bink_textures validated on the game thread.
+struct FBinkOverlayRenderJob
+{
+	FBinkTextures* BinkTextures;
+	float ulx, uly, lrx, lry;
+	int   tonemap;
+	int   out_nits;
+};
+
+static void Bink_DrawOverlays_Internal(
+	FRHICommandListImmediate& RHICmdList,
+	FTextureRHIRef BackBuffer,
+	TArray<FBinkOverlayRenderJob> Jobs)
 {
-	if (!BackBuffer) 
+	if (!BackBuffer.GetReference() || Jobs.IsEmpty())
 	{
 		return;
 	}
 
-	BINKPLUGINFRAMEINFO FrameInfo = {};
-	FrameInfo.screen_resource = BackBuffer;
-	FrameInfo.screen_resource_state = 4; // D3D12_RESOURCE_STATE_RENDER_TARGET; (only used in d3d12)
-	FrameInfo.width = ScreenSize.X;
-	FrameInfo.height = ScreenSize.Y;
-	FrameInfo.sdr_or_hdr = BackBuffer->GetFormat() == PF_A2B10G10R10 ? 1 : 0;
-	FrameInfo.cmdBuf = &RHICmdList;
-	BinkPluginSetPerFrameInfo(&FrameInfo);
-	BinkPluginAllScheduled();
-	BinkPluginDraw(0, 1);
+	for (FBinkOverlayRenderJob& job : Jobs)
+	{
+		if (!job.BinkTextures)
+		{
+			continue;
+		}
+
+		// Determine HDR tonemapping on the render thread where CVar reads are correct
+		int32 tonemap, out_nits;
+		BinkGetHDRSettings(tonemap, out_nits);
+
+		// Overlay draws load (blend over) existing backbuffer content
+		job.BinkTextures->SetRenderTarget(BackBuffer.GetReference(), ERenderTargetLoadAction::ELoad);
+
+		job.BinkTextures->SetHDRSettings(tonemap, 1.0f, out_nits);
+		job.BinkTextures->SetAlphaSettings(1.0f, 0);
+		job.BinkTextures->SetDrawPosition(job.ulx, job.uly, job.lrx, job.lry);
+
+		job.BinkTextures->Draw();
+	}
+
 }
 
 void UBinkFunctionLibrary::Bink_DrawOverlays() 
 {
+	if (PendingBinkOverlays.IsEmpty())
+	{
+		return;
+	}
+
 	if (!GEngine || !GEngine->GameViewport || !GEngine->GameViewport->Viewport)
 	{
+		PendingBinkOverlays.Reset();
 		return;
 	}
 
-	ENQUEUE_RENDER_COMMAND(BinkOverlays)(
-		[Viewport = GEngine->GameViewport->Viewport](FRHICommandListImmediate& RHICmdList)
+	// Collect render jobs — read bink_textures NOW (game thread, after all Ticks).
+	// Any Close() that ran during Tick already set bink_textures to null, so we skip those.
+	TArray<FBinkOverlayRenderJob> RenderJobs;
+	RenderJobs.Reserve(PendingBinkOverlays.Num());
+
+	for (const FBinkOverlayJob& job : PendingBinkOverlays)
+	{
+		FBinkTextures* tex = job.Player ? job.Player->bink_textures : nullptr;
+		if (tex)
+		{
+			RenderJobs.Add({ tex, job.ulx, job.uly, job.lrx, job.lry, 0, 80 });
+		}
+	}
+
+	PendingBinkOverlays.Reset();
+
+	if (RenderJobs.IsEmpty())
+	{
+		return;
+	}
+
+	FVector2D ScreenSize;
+	GEngine->GameViewport->GetViewportSize(ScreenSize);
+	FTextureRHIRef BackBuffer = GEngine->GameViewport->Viewport->GetRenderTargetTexture();
+
+	if (BackBuffer.IsValid())
+	{
+		ENQUEUE_RENDER_COMMAND(BinkOverlays)(
+			[BackBuffer, Jobs = MoveTemp(RenderJobs)](FRHICommandListImmediate& RHICmdList) mutable
 		{
-			Bink_DrawOverlays_Internal(RHICmdList, Viewport->GetRenderTargetTexture(), Viewport->GetSizeXY());
+			Bink_DrawOverlays_Internal(RHICmdList, BackBuffer, MoveTemp(Jobs));
 		});
+	}
+	else
+	{
+		TSharedPtr<SWindow> Window = FSlateApplication::Get().GetActiveTopLevelWindow();
+		ISlateViewportProvider* ViewportProvider = Window.IsValid() ? FSlateApplication::Get().GetRenderer()->GetViewportProvider(*Window) : nullptr;
+		if (ViewportProvider)
+		{
+			ENQUEUE_RENDER_COMMAND(BinkOverlays)(
+				[ViewportProvider, Jobs = MoveTemp(RenderJobs)](FRHICommandListImmediate& RHICmdList) mutable
+			{
+				FRHITexture* BB = ViewportProvider->GetBackBufferResource();
+				if (BB)
+				{
+					FTextureRHIRef BBRef = BB;
+					Bink_DrawOverlays_Internal(RHICmdList, BBRef, MoveTemp(Jobs));
+				}
+			});
+		}
+	}
 }
 
 FTimespan UBinkFunctionLibrary::BinkLoadingMovie_GetDuration() 
 {
 	double ms = 0;
-	if(MovieStreamer.IsValid() && MovieStreamer.Get()->bnk) 
+	if (MovieStreamer.IsValid() && MovieStreamer.Get()->bnk) 
 	{
-		BINKPLUGININFO bpinfo = {};
-		BinkPluginInfo(MovieStreamer.Get()->bnk, &bpinfo);
-		ms = ((double)bpinfo.Frames) * ((double)bpinfo.FrameRateDiv) * 1000.0 / ((double)bpinfo.FrameRate);
+		BINKHLINFO info = {};
+		BinkHLGetInfo(MovieStreamer.Get()->bnk, &info);
+		ms = ((double)info.Frames) * ((double)info.FrameRateDiv) * 1000.0 / ((double)info.FrameRate);
 	}
 	return FTimespan::FromMilliseconds(ms);
 }
@@ -62,11 +140,11 @@ FTimespan UBinkFunctionLibrary::BinkLoadingMovie_GetDuration()
 FTimespan UBinkFunctionLibrary::BinkLoadingMovie_GetTime() 
 {
 	double ms = 0;
-	if(MovieStreamer.IsValid() && MovieStreamer.Get()->bnk) 
+	if (MovieStreamer.IsValid() && MovieStreamer.Get()->bnk) 
 	{
-		BINKPLUGININFO bpinfo = {};
-		BinkPluginInfo(MovieStreamer.Get()->bnk, &bpinfo);
-		ms = ((double)bpinfo.FrameNum) * ((double)bpinfo.FrameRateDiv) * 1000.0 / ((double)bpinfo.FrameRate);
+		BINKHLINFO info = {};
+		BinkHLGetInfo(MovieStreamer.Get()->bnk, &info);
+		ms = ((double)info.FrameNum) * ((double)info.FrameRateDiv) * 1000.0 / ((double)info.FrameRate);
 	}
 	return FTimespan::FromMilliseconds(ms);
 }
diff --git a/Engine/Plugins/Media/BinkMedia/Source/BinkMediaPlayer/Private/BinkMediaPlayerModule.cpp b/Engine/Plugins/Media/BinkMedia/Source/BinkMediaPlayer/Private/BinkMediaPlayerModule.cpp
--- a/Engine/Plugins/Media/BinkMedia/Source/BinkMediaPlayer/Private/BinkMediaPlayerModule.cpp
+++ b/Engine/Plugins/Media/BinkMedia/Source/BinkMediaPlayer/Private/BinkMediaPlayerModule.cpp
@@ -10,166 +10,139 @@
 #include "RenderingThread.h"
 #include "HAL/PlatformFileManager.h"
 #include "UObject/UObjectIterator.h"
-#include "binkplugin_ue4.h"
-#include "egttypes.h"
+#include "BinkUE.h"
+#include "HDRHelper.h"
 
 DEFINE_LOG_CATEGORY(LogBink);
 
 TSharedPtr<FBinkMovieStreamer, ESPMode::ThreadSafe> MovieStreamer;
 
-TArray< FTextureRHIRef > BinkActiveTextureRefs;
-
-#if BINKPLUGIN_UE4_EDITOR
-class UFactory;
-#include "DetailLayoutBuilder.h"
-#include "Editor.h"
-#include "EditorDirectories.h"
-#include "Framework/Notifications/NotificationManager.h"
-#include "IDetailChildrenBuilder.h"
-#include "IDetailCustomization.h"
-#include "ISettingsModule.h"
-#include "Interfaces/IMainFrameModule.h"
-#include "PropertyCustomizationHelpers.h"
-#include "SourceControlHelpers.h"
-#include "Widgets/Input/SFilePathPicker.h"
-#include "Widgets/Notifications/SNotificationList.h"
-
-#include "BinkMediaPlayer.h"
-
-#define LOCTEXT_NAMESPACE "BinkMediaPlayerModule"
-
-struct FBinkMoviePlayerSettingsDetails : public IDetailCustomization {
-	static TSharedRef<IDetailCustomization> MakeInstance() { return MakeShareable(new FBinkMoviePlayerSettingsDetails); }
-	virtual void CustomizeDetails( IDetailLayoutBuilder& DetailLayout ) override {
-		IDetailCategoryBuilder& MoviesCategory = DetailLayout.EditCategory("Movies");
+EPixelFormat bink_force_pixel_format = PF_Unknown;
+static bool  bink_initialized        = false;
+TArray<FBinkOverlayJob> PendingBinkOverlays;
 
-		StartupMoviesPropertyHandle = DetailLayout.GetProperty("StartupMovies");
+// ---------------------------------------------------------------------------
+// Shared helper implementations
+// ---------------------------------------------------------------------------
 
-		TSharedRef<FDetailArrayBuilder> StartupMoviesBuilder = MakeShareable(new FDetailArrayBuilder(StartupMoviesPropertyHandle.ToSharedRef()));
-		StartupMoviesBuilder->OnGenerateArrayElementWidget(FOnGenerateArrayElementWidget::CreateSP(this, &FBinkMoviePlayerSettingsDetails::GenerateArrayElementWidget));
+void BinkGetHDRSettings(int32& OutTonemap, int32& OutNits)
+{
+	static const auto CVarHDROutputEnabled    = IConsoleManager::Get().FindTConsoleVariableDataInt(TEXT("r.HDR.EnableHDROutput"));
+	static const auto CVarDisplayOutputDevice = IConsoleManager::Get().FindTConsoleVariableDataInt(TEXT("r.HDR.Display.OutputDevice"));
 
-		MoviesCategory.AddProperty("bWaitForMoviesToComplete");
-		MoviesCategory.AddProperty("bMoviesAreSkippable");
+	OutTonemap = 1;
+	OutNits    = 80;
 
-		const bool bForAdvanced = false;
-		MoviesCategory.AddCustomBuilder(StartupMoviesBuilder, bForAdvanced);
+	if (GRHISupportsHDROutput && CVarHDROutputEnabled->GetValueOnRenderThread() != 0)
+	{
+		EDisplayOutputFormat outDev = static_cast<EDisplayOutputFormat>(CVarDisplayOutputDevice->GetValueOnRenderThread());
+		float DisplayMaxLuminance  = HDRGetDisplayMaximumLuminance();
+		switch (outDev)
+		{
+		case EDisplayOutputFormat::SDR_sRGB:
+		case EDisplayOutputFormat::SDR_Rec709:
+		case EDisplayOutputFormat::SDR_ExplicitGammaMapping:
+			OutTonemap = 1;
+			OutNits    = 80;
+			break;
+		case EDisplayOutputFormat::HDR_ACES_1000nit_ST2084:
+			OutTonemap = 2;
+			OutNits    = (int32)DisplayMaxLuminance;
+			break;
+		case EDisplayOutputFormat::HDR_ACES_1000nit_ScRGB:
+			OutTonemap = 1;
+			OutNits    = (int32)DisplayMaxLuminance;
+			break;
+		case EDisplayOutputFormat::HDR_ACES_2000nit_ST2084:
+			OutTonemap = 2;
+			OutNits    = (int32)DisplayMaxLuminance;
+			break;
+		case EDisplayOutputFormat::HDR_ACES_2000nit_ScRGB:
+			OutTonemap = 1;
+			OutNits    = (int32)DisplayMaxLuminance;
+			break;
+		default:
+			OutTonemap = 0;
+			OutNits    = 1000;
+			break;
+		}
 	}
+}
 
-	void GenerateArrayElementWidget(TSharedRef<IPropertyHandle> PropertyHandle, int32 ArrayIndex, IDetailChildrenBuilder& ChildrenBuilder) {
-		IDetailPropertyRow& FilePathRow = ChildrenBuilder.AddProperty(PropertyHandle);
-		FilePathRow.CustomWidget(false)
-			.NameContent()
-			[
-				PropertyHandle->CreatePropertyNameWidget()
-			]
-		.ValueContent()
-			.MaxDesiredWidth(0.0f)
-			.MinDesiredWidth(125.0f)
-			[
-				SNew(SFilePathPicker)
-					.BrowseButtonImage(FAppStyle::GetBrush("PropertyWindow.Button_Ellipsis"))
-					.BrowseButtonStyle(FAppStyle::Get(), "HoverHintOnly")
-					.BrowseButtonToolTip(LOCTEXT("FileButtonToolTipText", "Choose a file from this computer"))
-					.BrowseDirectory(FEditorDirectories::Get().GetLastDirectory(ELastDirectory::GENERIC_OPEN))
-					.BrowseTitle(LOCTEXT("PropertyEditorTitle", "File picker..."))
-					.FilePath(this, &FBinkMoviePlayerSettingsDetails::HandleFilePathPickerFilePath, PropertyHandle)
-					.FileTypeFilter(TEXT("Bink 2 Movie (*.bk2)|*.bk2"))
-					.OnPathPicked(this, &FBinkMoviePlayerSettingsDetails::HandleFilePathPickerPathPicked, PropertyHandle)
-			];
-	}
+void BinkCloseOnRenderThread(struct BINK*& InOutBnk, struct FBinkTextures*& InOutTextures)
+{
+	HBINK OldBnk                = InOutBnk;
+	FBinkTextures* OldTextures  = InOutTextures;
+	InOutBnk      = nullptr;
+	InOutTextures = nullptr;
 
-	FString HandleFilePathPickerFilePath( TSharedRef<IPropertyHandle> Property ) const {
-		FString FilePath;
-		Property->GetValue(FilePath);
-		return FilePath;
+	if (OldBnk || OldTextures)
+	{
+		ENQUEUE_RENDER_COMMAND(BinkClose)([OldBnk, OldTextures](FRHICommandListImmediate& RHICmdList)
+		{
+			if (OldTextures) { OldTextures->Destroy(); }
+			if (OldBnk) { BinkClose(OldBnk); }
+		});
 	}
+}
 
-	void HandleFilePathPickerPathPicked( const FString& PickedPath, TSharedRef<IPropertyHandle> Property ) {
-		const FString MoviesBaseDir = FPaths::ConvertRelativePathToFull( BINKMOVIEPATH );
-		const FString FullPath = FPaths::ConvertRelativePathToFull(PickedPath);
-
-		if (FullPath.StartsWith(MoviesBaseDir)) {
-			FText FailReason;
-			if (SourceControlHelpers::CheckoutOrMarkForAdd(PickedPath, LOCTEXT("MovieFileDescription", "movie"), FOnPostCheckOut(), FailReason)) {
-				Property->SetValue(FPaths::GetBaseFilename(FullPath.RightChop(MoviesBaseDir.Len())));
-			} else {
-				FNotificationInfo Info(FailReason);
-				Info.ExpireDuration = 3.0f;
-				FSlateNotificationManager::Get().AddNotification(Info);
-			}
-		} else if (!PickedPath.IsEmpty()) {
-			const FString FileName = FPaths::GetCleanFilename(PickedPath);
-			const FString DestPath = MoviesBaseDir / FileName;
-			FText FailReason;
-			if (SourceControlHelpers::CopyFileUnderSourceControl(DestPath, PickedPath, LOCTEXT("MovieFileDescription", "movie"), FailReason)) {
-				// trim the path so we just have a partial path with no extension (the movie player expects this)
-				Property->SetValue(FPaths::GetBaseFilename(DestPath.RightChop(MoviesBaseDir.Len())));
-			} else {
-				FNotificationInfo FailureInfo(FailReason);
-				FailureInfo.ExpireDuration = 3.0f;
-				FSlateNotificationManager::Get().AddNotification(FailureInfo);
-			}
-		} else {
-			Property->SetValue(PickedPath);
-		}
+#if PLATFORM_ANDROID
+FString BinkGetAndroidAPKPath()
+{
+	if (JNIEnv* env = FAndroidApplication::GetJavaEnv())
+	{
+		extern struct android_app* GNativeAndroidApp;
+		jclass clazz         = env->GetObjectClass(GNativeAndroidApp->activity->clazz);
+		jmethodID methodID   = env->GetMethodID(clazz, "getPackageCodePath", "()Ljava/lang/String;");
+		jobject result       = env->CallObjectMethod(GNativeAndroidApp->activity->clazz, methodID);
+		jboolean isCopy;
+		const char* apkPath  = env->GetStringUTFChars((jstring)result, &isCopy);
+		FString Path         = FString(UTF8_TO_TCHAR(apkPath));
+		env->ReleaseStringUTFChars((jstring)result, apkPath);
+		return Path;
 	}
+	return FString();
+}
+#endif
 
-	TSharedPtr<IPropertyHandle> StartupMoviesPropertyHandle;
-};
-#endif //BINKPLUGIN_UE4_EDITOR
+// ---------------------------------------------------------------------------
 
-unsigned bink_gpu_api;
-unsigned bink_gpu_api_hdr;
-EPixelFormat bink_force_pixel_format = PF_Unknown;
-static bool bink_initialized = false;
-
-#ifdef BINK_NDA_GPU_ALLOC
-extern void* BinkAllocGpu(UINTa Amt, U32 Align);
-extern void BinkFreeGpu(void* ptr, UINTa Amt);
-#else
-#define BinkAllocGpu 0
-#define BinkFreeGpu 0
-#endif
+#if BINKPLUGIN_UE4_EDITOR
+#include "Editor.h"
+#include "ISettingsModule.h"
+#include "Interfaces/IMainFrameModule.h"
+#include "BinkMediaPlayer.h"
 
-#ifdef BINK_NDA_CPU_ALLOC
-extern void* BinkAllocCpu(UINTa Amt, U32 Align);
-extern void BinkFreeCpu(void* ptr);
-#elif defined BINK_NO_CPU_ALLOC
-#define BinkAllocCpu 0
-#define BinkFreeCpu 0
-#else
-static void *BinkAllocCpu(UINTa Amt, U32 Align) { return FMemory::Malloc(Amt, Align); }
-static void BinkFreeCpu(void * ptr) { FMemory::Free(ptr); }
-#endif
+#define LOCTEXT_NAMESPACE "BinkMediaPlayerModule"
+#endif // BINKPLUGIN_UE4_EDITOR
+
+// Expose this function for now so we can control sound system behavior
+RADEXPFUNC void RADEXPLINK BinkImmediateCloseSoundSystem( S32 onoff );  // 1 = close once no videos playing (default on desktop), 0=close in BinkFreeGlobals (default on consoles)
 
-bool BinkInitialize() {
+bool BinkInitialize() 
+{
 	if (bink_initialized)
 	{
 		return true;
 	}
 
-	// TODO: make this an INI setting and/or configurable in Project Settings
-	//BinkPluginTurnOnGPUAssist();
-
-	static BINKPLUGININITINFO InitInfo = { 0 };
-	InitInfo.queue = 0;
-	InitInfo.physical_device = 0;
-	InitInfo.alloc = BinkAllocCpu;
-	InitInfo.free = BinkFreeCpu;
-	InitInfo.gpu_alloc = BinkAllocGpu;
-	InitInfo.gpu_free = BinkFreeGpu;
-	bink_gpu_api = BinkRHI;
+	int numSpeakers = GetNumSpeakers();
+	if (numSpeakers <= 0) 
+	{ 
+		numSpeakers = 2; 
+	}
 
-	bink_initialized = (bool)BinkPluginInit(0, &InitInfo, bink_gpu_api);
-	return bink_initialized;
+	BinkImmediateCloseSoundSystem(0);
+	BinkHLInit(4, 44100, (U32)numSpeakers);
+	bink_initialized = true;
+	return true;
 }
 
 FString BinkUE4CookOnTheFlyPath(FString path, const TCHAR *filename) 
 {
-	// If this isn't a shipping build, copy the file to our temp directory (so that cook-on-the-fly works)
-	FString toPath = FPaths::ConvertRelativePathToFull(BINKTEMPPATH) + filename;
+	FString toPath   = FPaths::ConvertRelativePathToFull(BINKTEMPPATH) + filename;
 	FString fromPath = path + filename;
-	toPath = toPath.Replace(TEXT("/./"), TEXT("/"));
+	toPath   = toPath.Replace(TEXT("/./"), TEXT("/"));
 	fromPath = fromPath.Replace(TEXT("/./"), TEXT("/"));
 	FPlatformFileManager::Get().GetPlatformFile().CopyFile(*toPath, *fromPath);
 	return toPath;
@@ -188,11 +161,10 @@ struct FBinkMediaPlayerModule : IModuleInterface, FTickableGameObject
 
 		if (!bPluginInitialized)
 		{
-			printf("Bink Error: %s\n", BinkPluginError());
+			printf("Bink Error: %s\n", BinkGetError());
 		}
 
 		MovieStreamer = MakeShareable(new FBinkMovieStreamer);
-
 		GetMoviePlayer()->RegisterMovieStreamer(MovieStreamer);
 
 #if BINKPLUGIN_UE4_EDITOR
@@ -205,11 +177,6 @@ struct FBinkMediaPlayerModule : IModuleInterface, FTickableGameObject
 				GetMutableDefault<UBinkMoviePlayerSettings>()
 			);
 
-			//SettingsModule->RegisterViewer("Project", *this);
-
-			static IMainFrameModule& MainFrameModule = FModuleManager::LoadModuleChecked<IMainFrameModule>("MainFrame");
-			MainFrameModule.OnMainFrameCreationFinished().AddRaw(this, &FBinkMediaPlayerModule::Initialize);
-		
 			FEditorDelegates::BeginPIE.AddRaw(this, &FBinkMediaPlayerModule::HandleEditorTogglePIE);
 			FEditorDelegates::EndPIE.AddRaw(this, &FBinkMediaPlayerModule::HandleEditorTogglePIE);
 		}
@@ -221,7 +188,7 @@ struct FBinkMediaPlayerModule : IModuleInterface, FTickableGameObject
 	{
 		if (bPluginInitialized)
 		{
-			BinkPluginShutdown();
+			BinkHLShutdown();
 			if (overlayHook.IsValid() && GEngine && GEngine->GameViewport)
 			{
 				GEngine->GameViewport->OnDrawn().Remove(overlayHook);
@@ -235,14 +202,6 @@ struct FBinkMediaPlayerModule : IModuleInterface, FTickableGameObject
 	}
 
 #if BINKPLUGIN_UE4_EDITOR
-	void Initialize(TSharedPtr<SWindow> InRootWindow, bool bIsRunningStartupDialog)
-	{
-		// This overrides the 
-		FPropertyEditorModule& PropertyModule = FModuleManager::LoadModuleChecked<FPropertyEditorModule>("PropertyEditor");
-		PropertyModule.RegisterCustomClassLayout("MoviePlayerSettings", FOnGetDetailCustomizationInstance::CreateStatic(&FBinkMoviePlayerSettingsDetails::MakeInstance));
-		PropertyModule.NotifyCustomizationModuleChanged();
-	}
-
 	void HandleEditorTogglePIE(bool bIsSimulating)
 	{
 		for (TObjectIterator<UBinkMediaPlayer> It; It; ++It)
@@ -263,18 +222,13 @@ struct FBinkMediaPlayerModule : IModuleInterface, FTickableGameObject
 
 	void DrawBinks() 
 	{
-		ENQUEUE_RENDER_COMMAND(BinkProcess)([](FRHICommandListImmediate& RHICmdList) 
-		{ 
-			BINKPLUGINFRAMEINFO FrameInfo = {};
-			BinkPluginSetPerFrameInfo(&FrameInfo);
-			BinkPluginProcessBinks(0);
-			BinkPluginAllScheduled();
-			if(BinkActiveTextureRefs.Num())
-			{
-				BinkPluginDraw(1, 0);
-				BinkActiveTextureRefs.Empty(256);
-			}
+		// Decode all playing Binks. This runs on the render thread before
+		// UpdateDeferredResource and before any editor texture render commands.
+		ENQUEUE_RENDER_COMMAND(BinkProcess)([](FRHICommandListImmediate& RHICmdList)
+		{
+			BinkHLProcess();
 		});
+
 		UBinkFunctionLibrary::Bink_DrawOverlays();
 	}
 
@@ -294,7 +248,7 @@ struct FBinkMediaPlayerModule : IModuleInterface, FTickableGameObject
 		}
 	}
 
-    FDelegateHandle overlayHook;
+	FDelegateHandle overlayHook;
 
 #if BINKPLUGIN_UE4_EDITOR
 	FDelegateHandle pieBeginHook;
@@ -305,4 +259,6 @@ struct FBinkMediaPlayerModule : IModuleInterface, FTickableGameObject
 
 IMPLEMENT_MODULE( FBinkMediaPlayerModule, BinkMediaPlayer )
 
+#if BINKPLUGIN_UE4_EDITOR
 #undef LOCTEXT_NAMESPACE
+#endif
diff --git a/Engine/Plugins/Media/BinkMedia/Source/BinkMediaPlayer/Private/BinkMediaPlayerPrivate.h b/Engine/Plugins/Media/BinkMedia/Source/BinkMediaPlayer/Private/BinkMediaPlayerPrivate.h
--- a/Engine/Plugins/Media/BinkMedia/Source/BinkMediaPlayer/Private/BinkMediaPlayerPrivate.h
+++ b/Engine/Plugins/Media/BinkMedia/Source/BinkMediaPlayer/Private/BinkMediaPlayerPrivate.h
@@ -15,16 +15,14 @@
 #include "Android/AndroidApplication.h"
 #endif
 
-
-
+#include "egttypes.h"
+#include "binktiny.h"
+#include "BinkRHI.h"
 
 enum EPixelFormat : uint8;
 
-extern BINKMEDIAPLAYER_API unsigned bink_gpu_api;
-extern BINKMEDIAPLAYER_API unsigned bink_gpu_api_hdr;
 extern BINKMEDIAPLAYER_API EPixelFormat bink_force_pixel_format;
 extern BINKMEDIAPLAYER_API FString BinkUE4CookOnTheFlyPath(FString path, const TCHAR *filename);
-extern BINKMEDIAPLAYER_API TArray< FTextureRHIRef > BinkActiveTextureRefs;
 extern BINKMEDIAPLAYER_API bool BinkInitialize();
 
 static int GetNumSpeakers() 
@@ -40,10 +38,35 @@ static int GetNumSpeakers()
 			return mix->GetNumDeviceChannels();
 		}
 	}
-	return 2; // default case... unknown just assume 2 speakers stereo setup.
+	return 2;
 }
 
 DECLARE_LOG_CATEGORY_EXTERN(LogBink, Log, All);
 
+// Per-frame overlay draw request, populated during UBinkMediaPlayer::Tick and
+// consumed by UBinkFunctionLibrary::Bink_DrawOverlays on the same game thread tick.
+struct FBinkOverlayJob
+{
+	class UBinkMediaPlayer* Player;
+	float ulx, uly, lrx, lry;
+};
+
+extern TArray<FBinkOverlayJob> PendingBinkOverlays;
 
+// ---------------------------------------------------------------------------
+// Shared helpers to eliminate duplicate logic across BinkMediaPlayer,
+// BinkFunctionLibrary, and BinkMovieStreamer.
+// ---------------------------------------------------------------------------
 
+// Queries HDR CVars on the render thread and returns tonemap mode + output nits.
+// Call from the render thread only.
+BINKMEDIAPLAYER_API void BinkGetHDRSettings(int32& OutTonemap, int32& OutNits);
+
+// Closes a Bink handle and frees its textures on the render thread.
+// Nulls out both pointers. Safe to call from the game thread.
+BINKMEDIAPLAYER_API void BinkCloseOnRenderThread(struct BINK*& InOutBnk, struct FBinkTextures*& InOutTextures);
+
+#if PLATFORM_ANDROID
+// Retrieves the APK code path via JNI. Returns empty string on failure.
+BINKMEDIAPLAYER_API FString BinkGetAndroidAPKPath();
+#endif
diff --git a/Engine/Plugins/Media/BinkMedia/Source/BinkMediaPlayer/Private/BinkMovieStreamer.cpp b/Engine/Plugins/Media/BinkMedia/Source/BinkMediaPlayer/Private/BinkMovieStreamer.cpp
--- a/Engine/Plugins/Media/BinkMedia/Source/BinkMediaPlayer/Private/BinkMovieStreamer.cpp
+++ b/Engine/Plugins/Media/BinkMedia/Source/BinkMediaPlayer/Private/BinkMovieStreamer.cpp
@@ -6,18 +6,17 @@
 #include "BinkMediaPlayerPrivate.h"
 #include "BinkMediaPlayer.h"
 #include "BinkMoviePlayerSettings.h"
-#include "HDRHelper.h"
 
 #include "HAL/PlatformFileManager.h"
 #include "Misc/Paths.h"
 #include "RenderingThread.h"
-#include "binkplugin_ue4.h"
-#include "egttypes.h"
+#include "BinkUE.h"
 
 DEFINE_LOG_CATEGORY(LogBinkMoviePlayer);
 
 FBinkMovieStreamer::FBinkMovieStreamer()
 	: bnk()
+	, bink_textures()
 {
 	MovieViewport = MakeShareable(new FMovieViewport());
 	PlaybackType = MT_Normal;
@@ -38,7 +37,7 @@ bool FBinkMovieStreamer::Init(const TArray<FString>& MoviePaths, TEnumAsByte<EMo
 	{
 		return false;
 	}
-	MovieIndex = -1;
+	MovieIndex   = -1;
 	PlaybackType = inPlaybackType;
 	StoredMoviePaths = MoviePaths;
 	return OpenNextMovie();
@@ -51,7 +50,7 @@ FString FBinkMovieStreamer::GetMovieName()
 
 bool FBinkMovieStreamer::IsLastMovieInPlaylist()
 {
-	return MovieIndex == StoredMoviePaths.Num() -1;
+	return MovieIndex == StoredMoviePaths.Num() - 1;
 }
 
 void FBinkMovieStreamer::ForceCompletion()
@@ -61,7 +60,6 @@ void FBinkMovieStreamer::ForceCompletion()
 
 bool FBinkMovieStreamer::Tick(FRHICommandListBase& RHICmdList, float DeltaTime)
 {
-	// Loops through all movies while not valid
 	if (!bnk)
 	{
 		CloseMovie();
@@ -80,9 +78,9 @@ bool FBinkMovieStreamer::Tick(FRHICommandListBase& RHICmdList, float DeltaTime)
 		}
 	}
 
-	BINKPLUGININFO bpinfo = {};
-	BinkPluginInfo(bnk, &bpinfo);
-	if (bpinfo.PlaybackState == 3)
+	BINKHLINFO info = {};
+	BinkHLGetInfo(bnk, &info);
+	if (info.PlaybackState == 3)
 	{
 		CloseMovie();
 		if (MovieIndex < StoredMoviePaths.Num() - 1)
@@ -104,13 +102,17 @@ bool FBinkMovieStreamer::Tick(FRHICommandListBase& RHICmdList, float DeltaTime)
 		return false;
 	}
 
+	// Decode the next frame for all playing Binks. During loading screens the game thread
+	// DrawBinks path may not fire, so the streamer calls BinkHLProcess itself.
+	BinkHLProcess();
+
 	FSlateTexture2DRHIRef* CurrentTexture = Texture.Get();
-	if(CurrentTexture)
+	if (CurrentTexture)
 	{
-		FVector2D destUpperLeft = GetDefault<UBinkMoviePlayerSettings>()->BinkDestinationUpperLeft;
+		FVector2D destUpperLeft  = GetDefault<UBinkMoviePlayerSettings>()->BinkDestinationUpperLeft;
 		FVector2D destLowerRight = GetDefault<UBinkMoviePlayerSettings>()->BinkDestinationLowerRight;
 
-		if( !CurrentTexture->IsInitialized() )
+		if (!CurrentTexture->IsInitialized())
 		{
 			CurrentTexture->InitResource(RHICmdList);
 		}
@@ -118,57 +120,24 @@ bool FBinkMovieStreamer::Tick(FRHICommandListBase& RHICmdList, float DeltaTime)
 		FTextureRHIRef tex = CurrentTexture->GetTypedResource();
 		uint32 binkw = tex.GetReference()->GetSizeX();
 		uint32 binkh = tex.GetReference()->GetSizeY();
-		bool is_hdr = tex.GetReference()->GetFormat() != PF_B8G8R8A8;
-
-		float ulx = destUpperLeft.X;
-		float uly = destUpperLeft.Y;
-		float lrx = destLowerRight.X;
-		float lry = destLowerRight.Y;
+		bool is_hdr  = tex.GetReference()->GetFormat() != PF_B8G8R8A8;
 
-		static const auto CVarHDROutputEnabled = IConsoleManager::Get().FindTConsoleVariableDataInt(TEXT("r.HDR.EnableHDROutput"));
-		static const auto CVarDisplayOutputDevice = IConsoleManager::Get().FindTConsoleVariableDataInt(TEXT("r.HDR.Display.OutputDevice"));
-		if (GRHISupportsHDROutput && CVarHDROutputEnabled->GetValueOnRenderThread() != 0)
-		{
-			EDisplayOutputFormat outDev = static_cast<EDisplayOutputFormat>(CVarDisplayOutputDevice->GetValueOnRenderThread());
-			float DisplayMaxLuminance = HDRGetDisplayMaximumLuminance();
-			switch (outDev)
-			{
-			// LDR
-			case EDisplayOutputFormat::SDR_sRGB:
-			case EDisplayOutputFormat::SDR_Rec709:
-			case EDisplayOutputFormat::SDR_ExplicitGammaMapping:
-				BinkPluginSetHdrSettings(bnk, 1, 1.0f, 80);
-				break;
-			// 1k nits
-			case EDisplayOutputFormat::HDR_ACES_1000nit_ST2084:
-			case EDisplayOutputFormat::HDR_ACES_1000nit_ScRGB:
-				BinkPluginSetHdrSettings(bnk, 1, 1.0f, DisplayMaxLuminance);
-				break;
-			// 2k nits
-			case EDisplayOutputFormat::HDR_ACES_2000nit_ST2084:
-			case EDisplayOutputFormat::HDR_ACES_2000nit_ScRGB:
-				BinkPluginSetHdrSettings(bnk, 1, 1.0f, DisplayMaxLuminance);
-				break;
-			// no tonemap
-			default:
-				BinkPluginSetHdrSettings(bnk, 0, 1.0f, 1000);
-				break;
-			}
-		}
-		else
+		if (bink_textures)
 		{
-			BinkPluginSetHdrSettings(bnk, 1, 1.0f, 80);
-		}
-		BinkPluginSetRenderTargetFormat(bnk, is_hdr ? 1 : 0);
+			int32 tonemap, out_nits;
+			BinkGetHDRSettings(tonemap, out_nits);
+
+			// The movie streamer runs its own decode+draw loop (no DrawBinks involvement),
+			// so draw directly here rather than queuing to PendingBinkTextureDraws.
+			bink_textures->SetHDRSettings(tonemap, 1.0f, out_nits);
+			bink_textures->SetDrawPosition(
+				destUpperLeft.X, destUpperLeft.Y,
+				destLowerRight.X, destLowerRight.Y);
 
-		BinkPluginScheduleToTexture(bnk, ulx, uly, lrx, lry, 0, tex.GetReference(), binkw, binkh);
+			bink_textures->SetRenderTarget(tex.GetReference(), ERenderTargetLoadAction::EClear);
 
-		BINKPLUGINFRAMEINFO FrameInfo = {};
-		FrameInfo.cmdBuf = &RHICmdList;
-		BinkPluginSetPerFrameInfo(&FrameInfo);
-		BinkPluginAllScheduled();
-		BinkPluginProcessBinks(0);
-		BinkPluginDraw(1, 0);
+			bink_textures->Draw();
+		}
 
 		MovieViewport->SetTexture(Texture);
 	}
@@ -179,9 +148,9 @@ bool FBinkMovieStreamer::Tick(FRHICommandListBase& RHICmdList, float DeltaTime)
 void FBinkMovieStreamer::Cleanup()
 {
 	FlushRenderingCommands();
-	for( int32 TextureIndex = 0; TextureIndex < TextureFreeList.Num(); ++TextureIndex )
+	for (int32 TextureIndex = 0; TextureIndex < TextureFreeList.Num(); ++TextureIndex)
 	{
-		BeginReleaseResource( TextureFreeList[TextureIndex].Get() );
+		BeginReleaseResource(TextureFreeList[TextureIndex].Get());
 	}
 }
 
@@ -190,13 +159,9 @@ bool FBinkMovieStreamer::OpenNextMovie()
 	MovieIndex++;
 	check(StoredMoviePaths.Num() > 0 && MovieIndex < StoredMoviePaths.Num());
 
-	if (GEngine) {
-		BinkPluginLimitSpeakers(GetNumSpeakers());
-	}
-
-	BINKPLUGINBUFFERING bufferMode = (BINKPLUGINBUFFERING)(int)GetDefault<UBinkMoviePlayerSettings>()->BinkBufferMode;
-	BINKPLUGINSNDTRACK soundTrack = (BINKPLUGINSNDTRACK)(int)GetDefault<UBinkMoviePlayerSettings>()->BinkSoundTrack;
-	int soundTrackStart = GetDefault<UBinkMoviePlayerSettings>()->BinkSoundTrackStart;
+	U32 bufferMode  = (U32)(int)GetDefault<UBinkMoviePlayerSettings>()->BinkBufferMode;
+	U32 soundTrack  = (U32)(int)GetDefault<UBinkMoviePlayerSettings>()->BinkSoundTrack;
+	U32 soundStart  = (U32)GetDefault<UBinkMoviePlayerSettings>()->BinkSoundTrackStart;
 
 	IPlatformFile& PlatformFile = FPlatformFileManager::Get().GetPlatformFile();
 
@@ -204,85 +169,82 @@ bool FBinkMovieStreamer::OpenNextMovie()
 	{
 		StoredMoviePaths[MovieIndex] + TEXT(".bk2")
 	};
-	for (int i = 0; i < sizeof(MoviePathTbl) / sizeof(MoviePathTbl[0]) && !bnk; ++i)
+	for (int i = 0; i < (int)(sizeof(MoviePathTbl) / sizeof(MoviePathTbl[0])) && !bnk; ++i)
 	{
 		FString FullMoviePath = BINKMOVIEPATH + MoviePathTbl[i];
 #if PLATFORM_ANDROID
-		// Don't bother trying to play it if we can't find it
 		if (!IAndroidPlatformFile::GetPlatformPhysical().FileExists(*FullMoviePath)) {
 			continue;
 		}
 
-		int64 FileOffset = IAndroidPlatformFile::GetPlatformPhysical().FileStartOffset(*FullMoviePath);
+		int64 FileOffset     = IAndroidPlatformFile::GetPlatformPhysical().FileStartOffset(*FullMoviePath);
 		FString FileRootPath = IAndroidPlatformFile::GetPlatformPhysical().FileRootPath(*FullMoviePath);
 
 		if (IAndroidPlatformFile::GetPlatformPhysical().IsAsset(*FullMoviePath))
 		{
-			if (JNIEnv* env = FAndroidApplication::GetJavaEnv())
+			FString APKPath = BinkGetAndroidAPKPath();
+			if (!APKPath.IsEmpty())
 			{
-				extern struct android_app* GNativeAndroidApp;
-				jclass clazz = env->GetObjectClass(GNativeAndroidApp->activity->clazz);
-				jmethodID methodID = env->GetMethodID(clazz, "getPackageCodePath", "()Ljava/lang/String;");
-				jobject result = env->CallObjectMethod(GNativeAndroidApp->activity->clazz, methodID);
-				jboolean isCopy;
-				const char *apkPath = env->GetStringUTFChars((jstring)result, &isCopy);
-				bnk = BinkPluginOpen(apkPath, soundTrack, soundTrackStart, bufferMode, FileOffset);
-				env->ReleaseStringUTFChars((jstring)result, apkPath);
+				bnk = BinkHLOpen(TCHAR_TO_UTF8(*APKPath), soundTrack, soundStart, bufferMode, (U64)FileOffset);
 			}
 		}
 		else
 		{
-			bnk = BinkPluginOpen(TCHAR_TO_ANSI(*FileRootPath), soundTrack, soundTrackStart, bufferMode, FileOffset);
+			bnk = BinkHLOpen(TCHAR_TO_UTF8(*FileRootPath), soundTrack, soundStart, bufferMode, (U64)FileOffset);
 		}
 #else
 		if (FPlatformFileManager::Get().GetPlatformFile().FileExists(*FullMoviePath)) 
 		{
-			const FString ExternalPath = PlatformFile.ConvertToAbsolutePathForExternalAppForRead(*FullMoviePath);
-#if PLATFORM_WINDOWS 
-			bnk = BinkPluginOpenUTF16((unsigned short*)*ExternalPath, soundTrack, soundTrackStart, bufferMode, 0);
-#else
-			bnk = BinkPluginOpen(TCHAR_TO_ANSI(*ExternalPath), soundTrack, soundTrackStart, bufferMode, 0);
-#endif
+			FString AbsPath = PlatformFile.ConvertToAbsolutePathForExternalAppForRead(*FullMoviePath);
+			bnk = BinkHLOpen(TCHAR_TO_UTF8(*AbsPath), soundTrack, soundStart, bufferMode, 0);
 		}
 		if (!bnk) 
 		{
 			FString CookPath = *BinkUE4CookOnTheFlyPath(FPaths::ConvertRelativePathToFull(BINKMOVIEPATH), *MoviePathTbl[i]);
 			if (FPlatformFileManager::Get().GetPlatformFile().FileExists(*CookPath)) 
 			{
-				const FString ExternalPath = PlatformFile.ConvertToAbsolutePathForExternalAppForRead(*CookPath);
-#if PLATFORM_WINDOWS 
-				bnk = BinkPluginOpenUTF16((unsigned short *)*ExternalPath, soundTrack, soundTrackStart, bufferMode, 0);
-#else
-				bnk = BinkPluginOpen(TCHAR_TO_ANSI(*ExternalPath), soundTrack, soundTrackStart, bufferMode, 0);
-#endif
+				FString AbsPath = PlatformFile.ConvertToAbsolutePathForExternalAppForRead(*CookPath);
+				bnk = BinkHLOpen(TCHAR_TO_UTF8(*AbsPath), soundTrack, soundStart, bufferMode, 0);
 			}
 		}
 #endif
 	}
+
 	if (!bnk)
 	{
 		return false;
 	}
 
-	FIntPoint VideoDimensions = FIntPoint::ZeroValue;
-	BINKPLUGININFO bpinfo = {};
-	BinkPluginInfo(bnk, &bpinfo);
-	VideoDimensions.X = bpinfo.Width;
-	VideoDimensions.Y = bpinfo.Height;
+	// Create CPU-side decode buffers immediately. GPU textures created lazily on first Tick.
+	bink_textures = FBinkTextures::CreateCPU(bnk, nullptr);
+
+	if (!bink_textures)
+	{
+		UE_LOG(LogBinkMoviePlayer, Error, TEXT("FBinkMovieStreamer::OpenNextMovie: Failed to create CPU decode buffers."));
+		BinkCloseOnRenderThread(bnk, bink_textures);
+		return false;
+	}
+
+	BINKHLINFO info = {};
+	BinkHLGetInfo(bnk, &info);
+	FIntPoint VideoDimensions = FIntPoint(info.Width, info.Height);
 
-	if( VideoDimensions.X > 0 && VideoDimensions.Y > 0 )
+	if (VideoDimensions.X > 0 && VideoDimensions.Y > 0)
 	{
-		EPixelFormat pixelFmt = bink_force_pixel_format == PF_Unknown ? (EPixelFormat)GetDefault<UBinkMoviePlayerSettings>()->BinkPixelFormat : bink_force_pixel_format;
-		if( TextureFreeList.Num() > 0 )
+		EPixelFormat pixelFmt = bink_force_pixel_format == PF_Unknown
+			? (EPixelFormat)GetDefault<UBinkMoviePlayerSettings>()->BinkPixelFormat
+			: bink_force_pixel_format;
+
+		if (TextureFreeList.Num() > 0)
 		{
 			Texture = TextureFreeList.Pop();
 
-			if( Texture->GetWidth() != VideoDimensions.X || Texture->GetHeight() != VideoDimensions.Y )
+			if (Texture->GetWidth() != VideoDimensions.X || Texture->GetHeight() != VideoDimensions.Y)
 			{
-				FSlateTexture2DRHIRef *ref = Texture.Get();
-				ENQUEUE_RENDER_COMMAND(UpdateMovieTexture)([ref,VideoDimensions](FRHICommandListImmediate& RHICmdList)
+				FSlateTexture2DRHIRef* ref = Texture.Get();
+				ENQUEUE_RENDER_COMMAND(UpdateMovieTexture)([ref, VideoDimensions](FRHICommandListImmediate& RHICmdList)
 				{
-					ref->Resize( VideoDimensions.X, VideoDimensions.Y );
+					ref->Resize(VideoDimensions.X, VideoDimensions.Y);
 				});
 			}
 		}
@@ -290,7 +252,7 @@ bool FBinkMovieStreamer::OpenNextMovie()
 		{
 			const bool bCreateEmptyTexture = true;
 			Texture = MakeShareable(new FSlateTexture2DRHIRef(VideoDimensions.X, VideoDimensions.Y, pixelFmt, NULL, TexCreate_RenderTargetable, bCreateEmptyTexture));
-			FSlateTexture2DRHIRef *ref = Texture.Get();
+			FSlateTexture2DRHIRef* ref = Texture.Get();
 
 			ENQUEUE_RENDER_COMMAND(InitMovieTexture)([ref](FRHICommandListImmediate& RHICmdList)
 			{
@@ -314,10 +276,9 @@ void FBinkMovieStreamer::CloseMovie()
 		MovieViewport->SetTexture(NULL);
 		Texture.Reset();
 	}
-	if(bnk)
+
+	if (bnk)
 	{
-		BinkPluginClose(bnk);
-		bnk = NULL;
+		BinkCloseOnRenderThread(bnk, bink_textures);
 	}
 }
-
diff --git a/Engine/Plugins/Media/BinkMedia/Source/BinkMediaPlayer/Private/BinkMovieStreamer.h b/Engine/Plugins/Media/BinkMedia/Source/BinkMediaPlayer/Private/BinkMovieStreamer.h
--- a/Engine/Plugins/Media/BinkMedia/Source/BinkMediaPlayer/Private/BinkMovieStreamer.h
+++ b/Engine/Plugins/Media/BinkMedia/Source/BinkMediaPlayer/Private/BinkMovieStreamer.h
@@ -1,10 +1,12 @@
 // Copyright Epic Games Tools LLC
-//   Licenced under the Unreal Engine EULA 
+//   Licenced under the Unreal Engine EULA
 #pragma once
 
-#include "binkplugin.h"
 #include "MoviePlayer.h"
 
+struct BINK;        // forward declare (HBINK = struct BINK*)
+struct FBinkTextures; // opaque Bink GPU texture state
+
 class FSlateTexture2DRHIRef;
 
 DECLARE_LOG_CATEGORY_EXTERN(LogBinkMoviePlayer, Log, All);
@@ -16,7 +18,6 @@ public:
 	virtual ~FBinkMovieStreamer();
 
 	virtual bool Init(const TArray<FString>& MoviePaths, TEnumAsByte<EMoviePlaybackType> inPlaybackType) override;
-	//virtual void PreviousViewportInterface(const TSharedPtr<ISlateViewport>& PreviousViewportInterface) override;
 	virtual void ForceCompletion() override;
 	virtual bool Tick(FRHICommandListBase& RHICmdList, float DeltaTime) override;
 	virtual TSharedPtr<class ISlateViewport> GetViewportInterface() override { return MovieViewport; }
@@ -26,32 +27,21 @@ public:
 	virtual void Cleanup() override;
 	virtual FTextureRHIRef GetTexture() override { return Texture->IsValid() ? Texture->GetRHIRef() : nullptr; }
 
-	//virtual void TickPreEngine() override;
-	//virtual void TickPostEngine() override;
-	//virtual void TickPostRender() override;
-
 	FOnCurrentMovieClipFinished OnCurrentMovieClipFinishedDelegate;
 	virtual FOnCurrentMovieClipFinished& OnCurrentMovieClipFinished() override { return OnCurrentMovieClipFinishedDelegate; }
 
 	bool OpenNextMovie();
 	void CloseMovie();
 
-	// A list of all the stored movie paths we have enqueued for playing
 	TArray<FString> StoredMoviePaths;
 
-	// Texture and viewport data for displaying to Slate
 	TSharedPtr<FMovieViewport> MovieViewport;
 	TSharedPtr<FSlateTexture2DRHIRef, ESPMode::ThreadSafe> Texture;
-
-	// List of textures pending deletion, need to keep this list because we can't immediately
-	// destroy them since they could be getting used on the main thread
 	TArray<TSharedPtr<FSlateTexture2DRHIRef, ESPMode::ThreadSafe>> TextureFreeList;
 
-	/** if true, this sequence will loop when finsihed */
 	TEnumAsByte<EMoviePlaybackType> PlaybackType;
-
-	// The index in to the playlist (StoredMoviePaths) that is currently playing
 	int32 MovieIndex;
 
-	BINKPLUGIN *bnk;
+	struct BINK* bnk;           // HBINK handle
+	FBinkTextures* bink_textures; // render-thread GPU texture state (lazily created)
 };
diff --git a/Engine/Plugins/Media/BinkMedia/Source/BinkMediaPlayer/Private/BinkRHI.h b/Engine/Plugins/Media/BinkMedia/Source/BinkMediaPlayer/Private/BinkRHI.h
new file mode 100644
--- /dev/null
+++ b/Engine/Plugins/Media/BinkMedia/Source/BinkMediaPlayer/Private/BinkRHI.h
@@ -0,0 +1,81 @@
+// Copyright Epic Games, Inc. All Rights Reserved.
+#pragma once
+
+#include "egttypes.h"
+#include "binktiny.h"
+#include "RHI.h"
+#include "RHICommandList.h"
+
+#define BINKMAXFRAMEBUFFERS 2
+
+// Bink GPU texture state — owns the CPU decode buffers and GPU textures
+// for a single playing Bink video. Created by CreateCPU() or Create(),
+// destroyed by Destroy(). All draw methods must be called on the render thread.
+struct FBinkTextures
+{
+	// --- Lifecycle -----------------------------------------------------------
+
+	// Full creation (CPU + GPU). Must be called on the render thread.
+	static FBinkTextures* Create(HBINK bink, void* user_ptr = nullptr);
+
+	// Creates only the CPU-side decode buffers and registers them with BinkHL.
+	// Safe to call from the game thread. GPU textures are created lazily on
+	// the render thread the first time Draw() is called.
+	static FBinkTextures* CreateCPU(HBINK bink, void* user_ptr = nullptr);
+
+	// Frees all GPU + CPU resources.
+	void Destroy();
+
+	// --- Per-frame draw setup (call before Draw) -----------------------------
+
+	void SetRenderTarget(FRHITexture* rt, ERenderTargetLoadAction load_action);
+	void SetDrawPosition(float x0, float y0, float x1, float y1);
+	void SetDrawCorners(float Ax, float Ay, float Bx, float By, float Cx, float Cy);
+	void SetSourceRect(float u0, float v0, float u1, float v1);
+	void SetAlphaSettings(float alpha, int32 draw_flags);
+	void SetHDRSettings(int32 tonemap, float exposure, int32 out_nits);
+
+	// Draws the current Bink frame to the previously-set render target.
+	void Draw();
+
+	// --- Internal state (treat as private) -----------------------------------
+
+	void* user_ptr = nullptr;
+	S32 video_width = 0, video_height = 0;
+	HBINK bink = nullptr;
+
+	U32 YABufferWidth = 0, YABufferHeight = 0;
+	U32 ChBufferWidth = 0, ChBufferHeight = 0;
+	S32 NeedAlpha = 0;
+	S32 NeedHDR = 0;
+
+	U8* Luma[BINKMAXFRAMEBUFFERS] = {};
+	U8* cR[BINKMAXFRAMEBUFFERS] = {};
+	U8* cB[BINKMAXFRAMEBUFFERS] = {};
+	U8* Alpha[BINKMAXFRAMEBUFFERS] = {};
+	U8* HDR[BINKMAXFRAMEBUFFERS] = {};
+
+	S32 bGpuTexturesCreated = 0;
+
+	F32 Ax = 0, Ay = 0, Bx = 0, By = 0, Cx = 0, Cy = 0;
+	F32 alpha = 0;
+	S32 draw_type = 0;
+	S32 draw_flags = 0;
+	F32 u0 = 0, v0 = 0, u1 = 0, v1 = 0;
+
+	FTextureRHIRef Ytexture[BINKMAXFRAMEBUFFERS];
+	FTextureRHIRef cRtexture[BINKMAXFRAMEBUFFERS];
+	FTextureRHIRef cBtexture[BINKMAXFRAMEBUFFERS];
+	FTextureRHIRef Atexture[BINKMAXFRAMEBUFFERS];
+	FTextureRHIRef Htexture[BINKMAXFRAMEBUFFERS];
+
+	S32 tonemap = 0;
+	F32 exposure = 0;
+	F32 out_luma = 0;
+
+	FRHITexture* RenderTarget = nullptr;
+	ERenderTargetLoadAction RenderTargetLoadAction = ERenderTargetLoadAction::EClear;
+
+private:
+	void EnsureGpuTextures();
+};
diff --git a/Engine/Plugins/Media/BinkMedia/Source/BinkMediaPlayer/Private/binkpluginRHI.cpp b/Engine/Plugins/Media/BinkMedia/Source/BinkMediaPlayer/Private/binkpluginRHI.cpp
--- a/Engine/Plugins/Media/BinkMedia/Source/BinkMediaPlayer/Private/binkpluginRHI.cpp
+++ b/Engine/Plugins/Media/BinkMedia/Source/BinkMediaPlayer/Private/binkpluginRHI.cpp
@@ -1,88 +1,66 @@
 // Copyright Epic Games, Inc. All Rights Reserved.
 
-#define BINKRHIFUNCTIONS
-#define BINKTEXTURESCLEANUP
 #include "egttypes.h"
-#include "RHIDefinitions.h"
 #include "binktiny.h"
-#include "RHIFwd.h"
-#include "binktextures.h"
-#include "binkplugin.h"
+#include "BinkRHI.h"
+#include "HAL/PlatformFileManager.h"
+#include "GenericPlatform/GenericPlatformFile.h"
 
 
-static BINKSHADERS * shaders;
 
-static int attached = 0;
-static unsigned rfb;
-
-FRHITexture *BinkRHIRenderTarget;
-ERenderTargetLoadAction BinkRenderTargetLoadAction;
-
-RADDEFFUNC int setup_rhi( void * device, BINKPLUGININITINFO * info, S32 gpu_assisted, void ** context )
+extern "C"
 {
-  if ( shaders == 0 )
-  {
-    shaders = Create_Bink_shaders( device );
-    if ( shaders == 0 )
-      return 0;
-  }
 
-  return 1;
-}
-
-RADDEFFUNC void shutdown_rhi( void )
+// Returns 1 on success, 0 on failure. Stores the file handle in user_data[0].
+S32 RADLINK binkdefopen(UINTa* user_data, const char* fn, U64* optional_out_file_size)
 {
-  if ( shaders )
-  {
-    Free_Bink_shaders( shaders );
-    shaders = 0;
-  }
-}
-
-RADDEFFUNC void * createtextures_rhi( void * bink )
-{
-  void * ret = 0;
-  if ( shaders )
-  {
-    ret = Create_Bink_textures( shaders, (HBINK)bink, 0 );
-  }
-  return ret;
-}
-
-RADDEFFUNC void selectrendertarget_rhi(void* texture_target, S32 width, S32 height, S32 do_clear, U32 sdr_or_hdr, S32 current_resource_state)
-{
-	if (texture_target)
+	FString Path = UTF8_TO_TCHAR(fn);
+	IFileHandle* Handle = FPlatformFileManager::Get().GetPlatformFile().OpenRead(*Path);
+	if (!Handle)
+	{
+		return 0;
+	}
+	user_data[0] = (UINTa)Handle;
+	if (optional_out_file_size)
 	{
-		attached = 1;
-		BinkRHIRenderTarget = (FRHITexture*)texture_target;
-		BinkRenderTargetLoadAction = do_clear ? ERenderTargetLoadAction::EClear : ERenderTargetLoadAction::ELoad;
+		*optional_out_file_size = (U64)Handle->Size();
 	}
+	return 1;
 }
 
-RADDEFFUNC void selectscreenrendertarget_rhi( void * texture_target, S32 width, S32 height, U32 format_idx, S32 current_resource_state )
+// Returns the number of bytes actually read.
+U64 RADLINK binkdefread(UINTa* user_data, void* dest, U64 bytes)
 {
-	if (texture_target)
+	IFileHandle* Handle = (IFileHandle*)user_data[0];
+	if (!Handle)
+	{
+		return 0;
+	}
+	// IFileHandle::Read takes int64; read in chunks if bytes > INT64_MAX (extremely unlikely)
+	int64 ToRead = (int64)bytes;
+	if (Handle->Read(static_cast<uint8*>(dest), ToRead))
 	{
-		attached = 2;
-		BinkRHIRenderTarget = (FRHITexture*)texture_target;
-		BinkRenderTargetLoadAction = ERenderTargetLoadAction::ELoad;
+		return (U64)ToRead;
 	}
+	return 0;
 }
 
-RADDEFFUNC void clearrendertarget_rhi(void)
+// Seeks to an absolute file position.
+void RADLINK binkdefseek(UINTa* user_data, U64 pos)
 {
-	if (attached)
+	IFileHandle* Handle = (IFileHandle*)user_data[0];
+	if (Handle)
 	{
-		BinkRHIRenderTarget = NULL;
-		attached = 0;
+		Handle->Seek((int64)pos);
 	}
 }
-  
-RADDEFFUNC void begindraw_rhi( void )
-{
-}
 
-RADDEFFUNC void enddraw_rhi( void )
+// void binkdefclose(UINTa* user_data)
+void RADLINK binkdefclose(UINTa* user_data)
 {
-  clearrendertarget_rhi();
+	IFileHandle* Handle = (IFileHandle*)user_data[0];
+	delete Handle;
+	user_data[0] = 0;
 }
+
+} // extern "C"
diff --git a/Engine/Plugins/Media/BinkMedia/Source/BinkMediaPlayer/Private/binktexturesRHI.cpp b/Engine/Plugins/Media/BinkMedia/Source/BinkMediaPlayer/Private/binktexturesRHI.cpp
--- a/Engine/Plugins/Media/BinkMedia/Source/BinkMediaPlayer/Private/binktexturesRHI.cpp
+++ b/Engine/Plugins/Media/BinkMedia/Source/BinkMediaPlayer/Private/binktexturesRHI.cpp
@@ -1,439 +1,306 @@
 // Copyright Epic Games, Inc. All Rights Reserved.
 
-#define BINKRHIFUNCTIONS
-#define BINKTEXTURESCLEANUP
+#include "BinkMediaPlayerPrivate.h"
 #include "egttypes.h"
 #include "binktiny.h"
-#include "binktextures.h"
-
+#include "BinkRHI.h"
 #include "BinkShaders.h"
 #include "PipelineStateCache.h"
 #include "RHIStaticStates.h"
 #include "RenderGraphBuilder.h"
 #include "RenderingThread.h"
 
-
-extern FRHITexture *BinkRHIRenderTarget;
-extern ERenderTargetLoadAction BinkRenderTargetLoadAction;
-
 FRDGTextureRef BinkRegisterExternalTexture(FRDGBuilder& GraphBuilder, FRHITexture* Texture, const TCHAR* NameIfUnregistered)
 {
-    if (FRDGTextureRef FoundTexture = GraphBuilder.FindExternalTexture(Texture))
-    {
-        return FoundTexture;
-    }
-    return GraphBuilder.RegisterExternalTexture(CreateRenderTarget(Texture, NameIfUnregistered));
+	if (FRDGTextureRef FoundTexture = GraphBuilder.FindExternalTexture(Texture))
+	{
+		return FoundTexture;
+	}
+	return GraphBuilder.RegisterExternalTexture(CreateRenderTarget(Texture, NameIfUnregistered));
 }
 
-typedef struct BINKSHADERSRHI
-{
-	BINKSHADERS pub;
-} BINKSHADERSRHI; 
 
-typedef struct BINKTEXTURESRHI
-{
-	BINKTEXTURES pub;
-
-	// RHI
-	S32 video_width, video_height;
-	S32 dirty;
-
-	BINKSHADERSRHI * shaders;
-	HBINK bink;
-	// this is the Bink info on the textures
-	BINKFRAMEBUFFERS bink_buffers;
-
-	F32 Ax, Ay, Bx, By, Cx, Cy;
-	F32 alpha;
-	S32 draw_type;
-	S32 draw_flags;
-	F32 u0, v0, u1, v1;
-
-	FTextureRHIRef Ytexture[BINKMAXFRAMEBUFFERS];
-	FTextureRHIRef cRtexture[BINKMAXFRAMEBUFFERS];
-	FTextureRHIRef cBtexture[BINKMAXFRAMEBUFFERS];
-	FTextureRHIRef Atexture[BINKMAXFRAMEBUFFERS];
-	FTextureRHIRef Htexture[BINKMAXFRAMEBUFFERS];
-
-	// unused for now
-	S32 tonemap;
-	F32 exposure;
-	F32 out_luma;
-
-} BINKTEXTURESRHI;
-
-#if defined( BINKTEXTURESINDIRECTBINKCALLS )
-  RADDEFFUNC void indirectBinkGetFrameBuffersInfo( HBINK bink, BINKFRAMEBUFFERS * fbset );
-  #define BinkGetFrameBuffersInfo indirectBinkGetFrameBuffersInfo
-  RADDEFFUNC void indirectBinkRegisterFrameBuffers( HBINK bink, BINKFRAMEBUFFERS * fbset );
-  #define BinkRegisterFrameBuffers indirectBinkRegisterFrameBuffers
-  RADDEFFUNC S32 indirectBinkAllocateFrameBuffers( HBINK bp, BINKFRAMEBUFFERS * set, U32 minimum_alignment );
-  #define BinkAllocateFrameBuffers indirectBinkAllocateFrameBuffers
-  RADDEFFUNC void * indirectBinkUtilMalloc(U64 bytes);
-  #define BinkUtilMalloc indirectBinkUtilMalloc
-  RADDEFFUNC void indirectBinkUtilFree(void * ptr);
-  #define BinkUtilFree indirectBinkUtilFree
-#endif
-
-static BINKTEXTURES * Create_textures( BINKSHADERS * pshaders, HBINK bink, void * user_ptr );
-static void Free_shaders( BINKSHADERS * pshaders );
-static void Free_textures( BINKTEXTURES * ptextures );
-static void Start_texture_update( BINKTEXTURES * ptextures );
-static void Finish_texture_update( BINKTEXTURES * ptextures );
-static void Draw_textures( BINKTEXTURES * ptextures, BINKSHADERS * pshaders, void * graphics_context );
-static void Set_draw_position( BINKTEXTURES * ptextures, F32 x0, F32 y0, F32 x1, F32 y1 );
-static void Set_draw_corners( BINKTEXTURES * ptextures, F32 Ax, F32 Ay, F32 Bx, F32 By, F32 Cx, F32 Cy );
-static void Set_source_rect( BINKTEXTURES * ptextures, F32 u0, F32 v0, F32 u1, F32 v1 );
-static void Set_alpha_settings( BINKTEXTURES * ptextures, F32 alpha_value, S32 draw_type );
-static void Set_hdr_settings( BINKTEXTURES * ptextures, S32 tonemap, F32 exposure, S32 out_nits  );
+// Struct definition is in BinkRHI.h
 
 //-----------------------------------------------------------------------------
-RADDEFFUNC BINKSHADERS * Create_Bink_shaders(void * dummy_device)
+FBinkTextures* FBinkTextures::Create(HBINK bink, void* user_ptr)
 {
-	BINKSHADERSRHI * pshaders;
-
-	pshaders = (BINKSHADERSRHI*)FMemory::Malloc(sizeof(*pshaders));
-	if (pshaders == 0)
+	FBinkTextures* textures = FBinkTextures::CreateCPU(bink, user_ptr);
+	if (!textures)
 	{
-		return 0;
+		return nullptr;
 	}
 
-	FMemory::Memset(pshaders, 0, sizeof(*pshaders));
-
-	pshaders->pub.Create_textures = Create_textures;
-	pshaders->pub.Free_shaders = Free_shaders;
-
-	return &pshaders->pub;
+	textures->EnsureGpuTextures();
+	return textures;
 }
 
-//-----------------------------------------------------------------------------
-static void Free_shaders( BINKSHADERS * pshaders )
+// Creates CPU-side decode buffers and registers them with BinkHL.
+// Safe to call from the game thread. GPU textures are created lazily
+// on the render thread the first time Draw_textures is called.
+FBinkTextures* FBinkTextures::CreateCPU(HBINK bink, void* user_ptr)
 {
-	FMemory::Free(pshaders);
-}
-
-//-----------------------------------------------------------------------------
-static BINKTEXTURES * Create_textures(BINKSHADERS * pshaders, HBINK bink, void * user_ptr)
-{
-	BINKTEXTURESRHI * textures;
-	BINKFRAMEBUFFERS * bb;
-
-	textures = (BINKTEXTURESRHI *)FMemory::Malloc(sizeof(*textures));
+	FBinkTextures* textures = (FBinkTextures*)FMemory::Malloc(sizeof(*textures));
 	if (textures == 0)
-		return 0;
-
-	FMemory::Memset(textures, 0, sizeof(*textures));
-
-	textures->pub.user_ptr = user_ptr;
-	textures->shaders = (BINKSHADERSRHI*)pshaders;
-	textures->bink = bink;
-	textures->video_width = bink->Width;
-	textures->video_height = bink->Height;
-
-	bb = &textures->bink_buffers;
-
-	BinkGetFrameBuffersInfo(bink, bb);
-
-#if !PLATFORM_HAS_DIRECT_TEXTURE_MEMORY_ACCESS
-	// allocate the system memory buffers if not allocated
-	if (!BinkAllocateFrameBuffers(bink, bb, 0))
 	{
-		FMemory::Free(textures);
 		return 0;
 	}
-#endif
 
-	FRHICommandListImmediate& RHICmdList = FRHICommandListImmediate::Get();
-	EPixelFormat format = PF_R8;
-#if PLATFORM_HAS_DIRECT_TEXTURE_MEMORY_ACCESS 
-	ETextureCreateFlags TexCreateFlags = TexCreate_CPUReadback | TexCreate_CPUWritable | TexCreate_OfflineProcessed | TexCreate_NoTiling;
-#else
-	ETextureCreateFlags TexCreateFlags = TexCreate_NoTiling;
-#endif
+	BINKHLINFO hlinfo = {};
+	BinkHLGetInfo(bink, &hlinfo);
 
-	const FRHITextureCreateDesc YADesc =
-		FRHITextureCreateDesc::Create2D(TEXT("BINK"), bb->YABufferWidth, bb->YABufferHeight, format)
-		.SetFlags(TexCreateFlags);
+	FMemory::Memset(textures, 0, sizeof(*textures));
+	textures->user_ptr = user_ptr;
+	textures->bink         = bink;
+	textures->video_width  = hlinfo.Width;
+	textures->video_height = hlinfo.Height;
 
-	const FRHITextureCreateDesc cRcBDesc =
-		FRHITextureCreateDesc::Create2D(TEXT("BINK"), bb->cRcBBufferWidth, bb->cRcBBufferHeight, format)
-		.SetFlags(TexCreateFlags);
 
-	for (int i = 0; i < bb->TotalFrames; ++i)
-	{
-		BINKFRAMEPLANESET *bp_src = &bb->Frames[i];
+	textures->YABufferWidth  = hlinfo.TextureBufferYAWidth;
+	textures->YABufferHeight = hlinfo.TextureBufferYAHeight;
+	textures->ChBufferWidth  = hlinfo.TextureBufferChWidth;
+	textures->ChBufferHeight = hlinfo.TextureBufferChHeight;
+	textures->NeedAlpha      = hlinfo.NeedAlpha;
+	textures->NeedHDR        = hlinfo.NeedHDR;
+	textures->bGpuTexturesCreated = 0;  // GPU textures created lazily on render thread
 
-		if (bp_src->YPlane.Allocate)
-		{
-			textures->Ytexture[i] = RHICmdList.CreateTexture(YADesc);
-#if PLATFORM_HAS_DIRECT_TEXTURE_MEMORY_ACCESS 
-			bp_src->YPlane.Buffer = RHICmdList.LockTexture2D(textures->Ytexture[i], 0, PLATFORM_DIRECT_TEXTURE_MEMORY_ACCESS_LOCK_MODE, bp_src->YPlane.BufferPitch, false, false);
-			bp_src->YPlane.BufferPitch = bp_src->YPlane.BufferPitch ? bp_src->YPlane.BufferPitch : ((bb->YABufferWidth+255)&-256);
-#endif
-		}
+	const U32 LumaSize   = textures->YABufferWidth * textures->YABufferHeight;
+	const U32 ChromaSize = textures->ChBufferWidth  * textures->ChBufferHeight;
 
-		if (bp_src->cRPlane.Allocate)
+	for (int i = 0; i < BINKMAXFRAMEBUFFERS; ++i)
+	{
+		textures->Luma[i]  = (U8*)FMemory::Malloc(LumaSize);
+		textures->cR[i]    = (U8*)FMemory::Malloc(ChromaSize);
+		textures->cB[i]    = (U8*)FMemory::Malloc(ChromaSize);
+		textures->Alpha[i] = textures->NeedAlpha ? (U8*)FMemory::Malloc(LumaSize)  : nullptr;
+		textures->HDR[i]   = textures->NeedHDR   ? (U8*)FMemory::Malloc(LumaSize)  : nullptr;
+
+		if (!textures->Luma[i] || !textures->cR[i] || !textures->cB[i]
+			|| (textures->NeedAlpha && !textures->Alpha[i])
+			|| (textures->NeedHDR   && !textures->HDR[i]))
 		{
-			textures->cRtexture[i] = RHICmdList.CreateTexture(cRcBDesc);
-#if PLATFORM_HAS_DIRECT_TEXTURE_MEMORY_ACCESS 
-			bp_src->cRPlane.Buffer = RHICmdList.LockTexture2D(textures->cRtexture[i], 0, PLATFORM_DIRECT_TEXTURE_MEMORY_ACCESS_LOCK_MODE, bp_src->cRPlane.BufferPitch, false, false);
-			bp_src->cRPlane.BufferPitch = bp_src->cRPlane.BufferPitch ? bp_src->cRPlane.BufferPitch : ((bb->cRcBBufferWidth+255)&-256);
-#endif
+			textures->Destroy();
+			return nullptr;
 		}
 
-		if (bp_src->cBPlane.Allocate)
-		{
-			textures->cBtexture[i] = RHICmdList.CreateTexture(cRcBDesc);
-#if PLATFORM_HAS_DIRECT_TEXTURE_MEMORY_ACCESS 
-			bp_src->cBPlane.Buffer = RHICmdList.LockTexture2D(textures->cBtexture[i], 0, PLATFORM_DIRECT_TEXTURE_MEMORY_ACCESS_LOCK_MODE, bp_src->cBPlane.BufferPitch, false, false);
-			bp_src->cBPlane.BufferPitch = bp_src->cBPlane.BufferPitch ? bp_src->cBPlane.BufferPitch : ((bb->cRcBBufferWidth+255)&-256);
-#endif
-		}
+		FMemory::Memset(textures->Luma[i],   0,   LumaSize);
+		FMemory::Memset(textures->cR[i],    128,  ChromaSize);
+		FMemory::Memset(textures->cB[i],    128,  ChromaSize);
+		if (textures->Alpha[i]) FMemory::Memset(textures->Alpha[i], 255, LumaSize);
+		if (textures->HDR[i])   FMemory::Memset(textures->HDR[i],     0, LumaSize);
+	}
 
-		if (bp_src->APlane.Allocate)
-		{
-			textures->Atexture[i] = RHICmdList.CreateTexture(YADesc);
-#if PLATFORM_HAS_DIRECT_TEXTURE_MEMORY_ACCESS 
-			bp_src->APlane.Buffer = RHICmdList.LockTexture2D(textures->Atexture[i], 0, PLATFORM_DIRECT_TEXTURE_MEMORY_ACCESS_LOCK_MODE, bp_src->APlane.BufferPitch, false, false);
-			bp_src->APlane.BufferPitch = bp_src->APlane.BufferPitch ? bp_src->APlane.BufferPitch : ((bb->YABufferWidth+255)&-256);
-#endif
-		}
+	U8* AlphaPtrs[BINKMAXFRAMEBUFFERS] = { textures->Alpha[0], textures->Alpha[1] };
+	U8* HDRPtrs[BINKMAXFRAMEBUFFERS]   = { textures->HDR[0],   textures->HDR[1]   };
 
-		if (bp_src->HPlane.Allocate)
-		{
-			textures->Htexture[i] = RHICmdList.CreateTexture(YADesc);
-#if PLATFORM_HAS_DIRECT_TEXTURE_MEMORY_ACCESS 
-			bp_src->HPlane.Buffer = RHICmdList.LockTexture2D(textures->Htexture[i], 0, PLATFORM_DIRECT_TEXTURE_MEMORY_ACCESS_LOCK_MODE, bp_src->HPlane.BufferPitch, false, false);
-			bp_src->HPlane.BufferPitch = bp_src->HPlane.BufferPitch ? bp_src->HPlane.BufferPitch : ((bb->YABufferWidth+255)&-256);
-#endif
-		}
-	}
+	BinkHLRegisterTextureBuffers(bink,
+		textures->Luma,
+		textures->cR,
+		textures->cB,
+		textures->NeedAlpha ? AlphaPtrs : nullptr,
+		textures->NeedHDR   ? HDRPtrs   : nullptr);
+
+	textures->SetDrawCorners(0, 0, 1, 0, 0, 1);
+	textures->SetSourceRect(0, 0, 1, 1);
+	textures->SetAlphaSettings(1, 0);
+	textures->SetHDRSettings(0, 1.0f, 80);
 
-	// Register our locked texture pointers with Bink
-	BinkRegisterFrameBuffers(bink, bb);
-
-	Set_draw_corners(&textures->pub, 0, 0, 1, 0, 0, 1);
-	Set_source_rect(&textures->pub, 0, 0, 1, 1);
-	Set_alpha_settings(&textures->pub, 1, 0);
-	Set_hdr_settings(&textures->pub, 0, 1.0f, 80);
-
-	textures->pub.Free_textures = Free_textures;
-	textures->pub.Start_texture_update = Start_texture_update;
-	textures->pub.Finish_texture_update = Finish_texture_update;
-	textures->pub.Draw_textures = Draw_textures;
-	textures->pub.Set_draw_position = Set_draw_position;
-	textures->pub.Set_draw_corners = Set_draw_corners;
-	textures->pub.Set_source_rect = Set_source_rect;
-	textures->pub.Set_alpha_settings = Set_alpha_settings;
-	textures->pub.Set_hdr_settings = Set_hdr_settings;
-
-	return &textures->pub;
+	return textures;
 }
 
-//-----------------------------------------------------------------------------
 
-static void Free_textures(BINKTEXTURES* ptextures)
+//-----------------------------------------------------------------------------
+// Creates GPU textures for an already-CPU-initialized FBinkTextures.
+// Must be called on the render thread.
+void FBinkTextures::EnsureGpuTextures()
 {
+	if (bGpuTexturesCreated)
+	{
+		return;
+	}
+
+	check(IsInRenderingThread());
+
 	FRHICommandListImmediate& RHICmdList = FRHICommandListImmediate::Get();
-	BINKTEXTURESRHI* textures = (BINKTEXTURESRHI*)ptextures;
-	BINKFRAMEBUFFERS *bb = &textures->bink_buffers;
+	EPixelFormat Format  = PF_R8;
+	ETextureCreateFlags TexFlags = TexCreate_NoTiling;
+
+	const FRHITextureCreateDesc YADesc =
+		FRHITextureCreateDesc::Create2D(TEXT("BINK"), YABufferWidth, YABufferHeight, Format)
+		.SetFlags(TexFlags);
+	const FRHITextureCreateDesc ChDesc =
+		FRHITextureCreateDesc::Create2D(TEXT("BINK"), ChBufferWidth, ChBufferHeight, Format)
+		.SetFlags(TexFlags);
 
-	for (int i = 0; i < bb->TotalFrames; ++i)
+	for (int i = 0; i < BINKMAXFRAMEBUFFERS; ++i)
 	{
-		BINKFRAMEPLANESET *bp_src = &bb->Frames[i];
-#if PLATFORM_HAS_DIRECT_TEXTURE_MEMORY_ACCESS 
-		if (bp_src->YPlane.Buffer)  RHICmdList.UnlockTexture2D(textures->Ytexture[i], 0, false);
-		if (bp_src->cRPlane.Buffer) RHICmdList.UnlockTexture2D(textures->cRtexture[i], 0, false);
-		if (bp_src->cBPlane.Buffer) RHICmdList.UnlockTexture2D(textures->cBtexture[i], 0, false);
-		if (bp_src->APlane.Buffer)  RHICmdList.UnlockTexture2D(textures->Atexture[i], 0, false);
-		if (bp_src->HPlane.Buffer)  RHICmdList.UnlockTexture2D(textures->Htexture[i], 0, false);
-#endif
-		if (textures->Ytexture[i].IsValid())  textures->Ytexture[i].SafeRelease();
-		if (textures->cRtexture[i].IsValid()) textures->cRtexture[i].SafeRelease();
-		if (textures->cBtexture[i].IsValid()) textures->cBtexture[i].SafeRelease();
-		if (textures->Atexture[i].IsValid())  textures->Atexture[i].SafeRelease();
-		if (textures->Htexture[i].IsValid())  textures->Htexture[i].SafeRelease();
+		Ytexture[i]  = RHICmdList.CreateTexture(YADesc);
+		cRtexture[i] = RHICmdList.CreateTexture(ChDesc);
+		cBtexture[i] = RHICmdList.CreateTexture(ChDesc);
+		if (NeedAlpha) Atexture[i] = RHICmdList.CreateTexture(YADesc);
+		if (NeedHDR)   Htexture[i] = RHICmdList.CreateTexture(YADesc);
 	}
 
-	FMemory::Free(textures);
+	bGpuTexturesCreated = 1;
 }
 
-static void Start_texture_update( BINKTEXTURES * ptextures )
-{
-	//BINKTEXTURESRHI *textures = (BINKTEXTURESRHI*)ptextures;
-}
-
-static void Finish_texture_update( BINKTEXTURES * ptextures )
-{
-	BINKTEXTURESRHI * textures = (BINKTEXTURESRHI*)ptextures;
-	BINKFRAMEBUFFERS *bb = &textures->bink_buffers;
-	int frame_num = bb->FrameNum;
-	BINKFRAMEPLANESET *bp_src = &bb->Frames[frame_num];
-
-#if !PLATFORM_HAS_DIRECT_TEXTURE_MEMORY_ACCESS
-	textures->dirty = 1;
-#endif
-
-	textures->pub.Ytexture = (void*)(UINTa)&textures->Ytexture[frame_num];
-	textures->pub.cRtexture = (void*)(UINTa)&textures->cRtexture[frame_num];
-	textures->pub.cBtexture = (void*)(UINTa)&textures->cBtexture[frame_num];
-	textures->pub.Atexture = (void*)(UINTa)&textures->Atexture[frame_num];
-	textures->pub.Htexture = (void*)(UINTa)&textures->Htexture[frame_num];
-}
 
 //-----------------------------------------------------------------------------
-
-static void update_plane_texture_rect(FRHICommandListImmediate& RHI, FRHITexture* RHITexture, BINKPLANE const* plane, unsigned w, unsigned h)
+void FBinkTextures::Destroy()
 {
-	uint32 Stride = 0;
-	unsigned char* TextureMemory = (unsigned char*)RHI.LockTexture2D(RHITexture, 0, RLM_WriteOnly, Stride, false);
 
-	if (TextureMemory)
+	// Clear the flag BEFORE releasing GPU resources so that any concurrent reader
+	// seeing this struct (e.g. a stale PendingBinkTextureDraws entry) will know
+	// the GPU textures are gone and won't try to draw with them.
+	bGpuTexturesCreated = 0;
+
+	for (int i = 0; i < BINKMAXFRAMEBUFFERS; ++i)
 	{
-		check(Stride >= w);
-		if (Stride == w)
-		{
-			FMemory::Memcpy(TextureMemory, plane->Buffer, w * h);
-		}
-		else
-		{
-			for (unsigned i = 0; i < h; ++i)
-			{
-				FMemory::Memcpy(TextureMemory + i * Stride, (char*)plane->Buffer + i * w, w);
-			}
-		}
-		RHI.UnlockTexture2D(RHITexture, 0, false);
+		if (Ytexture[i].IsValid())  Ytexture[i].SafeRelease();
+		if (cRtexture[i].IsValid()) cRtexture[i].SafeRelease();
+		if (cBtexture[i].IsValid()) cBtexture[i].SafeRelease();
+		if (Atexture[i].IsValid())  Atexture[i].SafeRelease();
+		if (Htexture[i].IsValid())  Htexture[i].SafeRelease();
+
+		if (Luma[i])  { FMemory::Free(Luma[i]);  Luma[i]  = nullptr; }
+		if (cR[i])    { FMemory::Free(cR[i]);    cR[i]    = nullptr; }
+		if (cB[i])    { FMemory::Free(cB[i]);    cB[i]    = nullptr; }
+		if (Alpha[i]) { FMemory::Free(Alpha[i]); Alpha[i] = nullptr; }
+		if (HDR[i])   { FMemory::Free(HDR[i]);   HDR[i]   = nullptr; }
 	}
+
+	FMemory::Free(this);
 }
 
-static void Draw_textures(BINKTEXTURES* ptextures, BINKSHADERS* pshaders, void* graphics_context)
+//-----------------------------------------------------------------------------
+void FBinkTextures::Draw()
 {
-	BINKTEXTURESRHI* textures = (BINKTEXTURESRHI*)ptextures;
-	BINKSHADERSRHI* bshaders = (BINKSHADERSRHI*)pshaders;
-	BINKFRAMEBUFFERS *bb = &textures->bink_buffers;
-	int frame_num = bb->FrameNum;
-	BINKFRAMEPLANESET *bp_src = &bb->Frames[frame_num];
-
-	// nowhere to render to?
-	if (BinkRHIRenderTarget == 0)
-		return;
 
-	if (bshaders == 0)
-		bshaders = textures->shaders;
+	if (RenderTarget == nullptr)
+	{
+		return;
+	}
 
-	int hasAPlane = bp_src->APlane.Allocate;
-	int hasHPlane = bp_src->HPlane.Allocate;
+	// Create GPU textures on first use (safe: always called from render thread)
+	EnsureGpuTextures();
 
-	FRDGBuilder BinkGraphBuilder(FRHICommandListImmediate::Get());
+	int hasAPlane = NeedAlpha;
+	int hasHPlane = NeedHDR;
 
-	FBinkParameters *consts = BinkGraphBuilder.AllocParameters<FBinkParameters>();
+	BINKHLINFO Info;
+	BinkHLGetInfo(bink, &Info);
+	if (Info.TextureDrawFrame == 0)
+	{
+		return;
+	}
+	int Idx = (int)Info.TextureDrawIndex;
+	if (Idx < 0 || Idx >= BINKMAXFRAMEBUFFERS)
+	{
+		UE_LOG(LogBink, Error, TEXT("FBinkTextures::Draw: TextureDrawIndex %d out of range [0, %d)"), Idx, BINKMAXFRAMEBUFFERS);
+		return;
+	}
 
-	FRDGTexture *BinkRDGTexture = BinkRegisterExternalTexture(BinkGraphBuilder, BinkRHIRenderTarget, TEXT("Bink_RT"));
-	consts->RenderTargets[0] = FRenderTargetBinding(BinkRDGTexture, BinkRenderTargetLoadAction, 0);
-	consts->tex0 = BinkRegisterExternalTexture(BinkGraphBuilder, textures->Ytexture[frame_num], TEXT("BinkY"));
-	consts->tex1 = BinkRegisterExternalTexture(BinkGraphBuilder, textures->cRtexture[frame_num], TEXT("BinkCr"));
-	consts->tex2 = BinkRegisterExternalTexture(BinkGraphBuilder, textures->cBtexture[frame_num], TEXT("BinkCb"));
-	consts->tex3 = BinkRegisterExternalTexture(BinkGraphBuilder, hasAPlane ? textures->Atexture[frame_num] : textures->Ytexture[frame_num], TEXT("BinkA"));
-	consts->tex4 = BinkRegisterExternalTexture(BinkGraphBuilder, hasHPlane ? textures->Htexture[frame_num] : textures->Ytexture[frame_num], TEXT("BinkH"));
+	FRDGBuilder BinkGraphBuilder(FRHICommandListImmediate::Get());
+	FBinkParameters* consts = BinkGraphBuilder.AllocParameters<FBinkParameters>();
+
+	FRDGTexture* BinkRDGTexture = BinkRegisterExternalTexture(BinkGraphBuilder, RenderTarget, TEXT("Bink_RT"));
+	consts->RenderTargets[0] = FRenderTargetBinding(BinkRDGTexture, RenderTargetLoadAction, 0);
+	consts->tex0 = BinkRegisterExternalTexture(BinkGraphBuilder, Ytexture[Idx],  TEXT("BinkY"));
+	consts->tex1 = BinkRegisterExternalTexture(BinkGraphBuilder, cRtexture[Idx], TEXT("BinkCr"));
+	consts->tex2 = BinkRegisterExternalTexture(BinkGraphBuilder, cBtexture[Idx], TEXT("BinkCb"));
+	consts->tex3 = BinkRegisterExternalTexture(BinkGraphBuilder, hasAPlane ? Atexture[Idx] : Ytexture[Idx], TEXT("BinkA"));
+	consts->tex4 = BinkRegisterExternalTexture(BinkGraphBuilder, hasHPlane ? Htexture[Idx] : Ytexture[Idx], TEXT("BinkH"));
 	consts->samp0 = TStaticSamplerState<SF_Bilinear>::GetRHI();
 	consts->samp1 = TStaticSamplerState<SF_Bilinear>::GetRHI();
 	consts->samp2 = TStaticSamplerState<SF_Bilinear>::GetRHI();
 	consts->samp3 = TStaticSamplerState<SF_Bilinear>::GetRHI();
 	consts->samp4 = TStaticSamplerState<SF_Bilinear>::GetRHI();
 
-	// Pixel shader alpha constants
-	if (textures->draw_type == 3)
+	if (draw_type == 3)
 	{
 		unsigned BinkColorA = 0xffffffff;
 		unsigned BinkColorB = 0xffc0c0c0;
-		consts->consta.X = ((F32)(S32)((BinkColorA >> 0) & 255)) * (1.0f / 255.0f);
-		consts->consta.Y = ((F32)(S32)((BinkColorA >> 8) & 255)) * (1.0f / 255.0f);
+		consts->consta.X = ((F32)(S32)((BinkColorA >>  0) & 255)) * (1.0f / 255.0f);
+		consts->consta.Y = ((F32)(S32)((BinkColorA >>  8) & 255)) * (1.0f / 255.0f);
 		consts->consta.Z = ((F32)(S32)((BinkColorA >> 16) & 255)) * (1.0f / 255.0f);
 		consts->consta.W = ((F32)(S32)((BinkColorA >> 24) & 255)) * (1.0f / 255.0f);
-		consts->crc.X = ((F32)(S32)((BinkColorB >> 0) & 255)) * (1.0f / 255.0f);
-		consts->crc.Y = ((F32)(S32)((BinkColorB >> 8) & 255)) * (1.0f / 255.0f);
+		consts->crc.X = ((F32)(S32)((BinkColorB >>  0) & 255)) * (1.0f / 255.0f);
+		consts->crc.Y = ((F32)(S32)((BinkColorB >>  8) & 255)) * (1.0f / 255.0f);
 		consts->crc.Z = ((F32)(S32)((BinkColorB >> 16) & 255)) * (1.0f / 255.0f);
 		consts->crc.W = ((F32)(S32)((BinkColorB >> 24) & 255)) * (1.0f / 255.0f);
-		consts->cbc.X = ((F32)textures->bink->Width) / 8.0f;
-		consts->cbc.Y = ((F32)textures->bink->Height) / 8.0f;
+		consts->cbc.X = ((F32)video_width)  / 8.0f;
+		consts->cbc.Y = ((F32)video_height) / 8.0f;
 		consts->cbc.Z = 0.f;
 		consts->cbc.W = 0.f;
 	}
 	else
 	{
-		// Pixel shader alpha constants
-		consts->consta.Z = consts->consta.Y = consts->consta.X = textures->draw_type == 1 ? textures->alpha : 1.0f;
-		consts->consta.W = textures->alpha;
+		consts->consta.Z = consts->consta.Y = consts->consta.X = draw_type == 1 ? alpha : 1.0f;
+		consts->consta.W = alpha;
 
 		if (hasHPlane)
 		{
 			// HDR stuff
-			consts->crc.X = textures->bink->ColorSpace[0];
-			consts->crc.Y = textures->exposure;
-			consts->crc.Z = textures->out_luma;
+			consts->crc.X = Info.cs0;
+			consts->crc.Y = exposure;
+			consts->crc.Z = out_luma;
 			consts->crc.W = 0;
-			consts->cbc.X = textures->bink->ColorSpace[1];
-			consts->cbc.Y = textures->bink->ColorSpace[2];
-			consts->cbc.Z = textures->bink->ColorSpace[3];
-			consts->cbc.W = textures->bink->ColorSpace[4];
+			consts->cbc.X = Info.cs1;
+			consts->cbc.Y = Info.cs2;
+			consts->cbc.Z = Info.cs3;
+			consts->cbc.W = Info.cs4;
 		}
 		else
 		{
 			// set the constants for the type of ycrcb we have
-			consts->crc.X = textures->bink->ColorSpace[0];
-			consts->crc.Y = textures->bink->ColorSpace[1];
-			consts->crc.Z = textures->bink->ColorSpace[2];
-			consts->crc.W = textures->bink->ColorSpace[3];
-			consts->cbc.X = textures->bink->ColorSpace[4];
-			consts->cbc.Y = textures->bink->ColorSpace[5];
-			consts->cbc.Z = textures->bink->ColorSpace[6];
-			consts->cbc.W = textures->bink->ColorSpace[7];
-			consts->adj.X = textures->bink->ColorSpace[8];
-			consts->adj.Y = textures->bink->ColorSpace[9];
-			consts->adj.Z = textures->bink->ColorSpace[10];
-			consts->adj.W = textures->bink->ColorSpace[11];
-			consts->yscale.X = textures->bink->ColorSpace[12];
-			consts->yscale.Y = textures->bink->ColorSpace[13];
-			consts->yscale.Z = textures->bink->ColorSpace[14];
-			consts->yscale.W = textures->bink->ColorSpace[15];
+			consts->crc.X = Info.cs0;
+			consts->crc.Y = Info.cs1;
+			consts->crc.Z = Info.cs2;
+			consts->crc.W = Info.cs3;
+			consts->cbc.X = Info.cs4;
+			consts->cbc.Y = Info.cs5;
+			consts->cbc.Z = Info.cs6;
+			consts->cbc.W = Info.cs7;
+			consts->adj.X = Info.cs8;
+			consts->adj.Y = Info.cs9;
+			consts->adj.Z = Info.cs10;
+			consts->adj.W = Info.cs11;
+			consts->yscale.X = Info.cs12;
+			consts->yscale.Y = Info.cs13;
+			consts->yscale.Z = Info.cs14;
+			consts->yscale.W = Info.cs15;
 		}
 	}
 
-	// Vertex shader constants
-	consts->xy_xform0.X = (textures->Bx - textures->Ax) * 2.0f;
-	consts->xy_xform0.Y = (textures->Cx - textures->Ax) * 2.0f;
-	consts->xy_xform0.Z = (textures->By - textures->Ay) * -2.0f;
-	consts->xy_xform0.W = (textures->Cy - textures->Ay) * -2.0f; // view space has +y = up, our coords have +y = down
-	consts->xy_xform1.X = textures->Ax * 2.0f - 1.0f;
-	consts->xy_xform1.Y = 1.0f - textures->Ay * 2.0f;
+	consts->xy_xform0.X = (Bx - Ax) * 2.0f;
+	consts->xy_xform0.Y = (Cx - Ax) * 2.0f;
+	consts->xy_xform0.Z = (By - Ay) * -2.0f;
+	consts->xy_xform0.W = (Cy - Ay) * -2.0f;
+	consts->xy_xform1.X = Ax * 2.0f - 1.0f;
+	consts->xy_xform1.Y = 1.0f - Ay * 2.0f;
 	consts->xy_xform1.Z = 0.0f;
 	consts->xy_xform1.W = 0.0f;
 
-	// UV matrix
 	{
-		F32 luma_u_scale = (F32)textures->video_width / (F32)bb->YABufferWidth;
-		F32 luma_v_scale = (F32)textures->video_height / (F32)bb->YABufferHeight;
-		F32 chroma_u_scale = (F32)(textures->video_width / 2) / (F32)bb->cRcBBufferWidth;
-		F32 chroma_v_scale = (F32)(textures->video_height / 2) / (F32)bb->cRcBBufferHeight;
-
-		// Set up matrix columns for UV transform (currently just scale+translate, could add rotation)
-		// X column
-		consts->uv_xform0.X = (textures->u1 - textures->u0) * luma_u_scale;
+		F32 luma_u   = (F32)video_width        / (F32)YABufferWidth;
+		F32 luma_v   = (F32)video_height       / (F32)YABufferHeight;
+		F32 chroma_u = (F32)(video_width  / 2) / (F32)ChBufferWidth;
+		F32 chroma_v = (F32)(video_height / 2) / (F32)ChBufferHeight;
+
+		consts->uv_xform0.X = (u1 - u0) * luma_u;
 		consts->uv_xform0.Y = 0.0f;
-		consts->uv_xform0.Z = (textures->u1 - textures->u0) * chroma_u_scale;
+		consts->uv_xform0.Z = (u1 - u0) * chroma_u;
 		consts->uv_xform0.W = 0.0f;
 
-		// Y column
 		consts->uv_xform1.X = 0.0f;
-		consts->uv_xform1.Y = (textures->v1 - textures->v0) * luma_v_scale;
+		consts->uv_xform1.Y = (v1 - v0) * luma_v;
 		consts->uv_xform1.Z = 0.0f;
-		consts->uv_xform1.W = (textures->v1 - textures->v0) * chroma_v_scale;
+		consts->uv_xform1.W = (v1 - v0) * chroma_v;
 
-		// W column (translation)
-		consts->uv_xform2.X = textures->u0 * luma_u_scale;
-		consts->uv_xform2.Y = textures->v0 * luma_v_scale;
-		consts->uv_xform2.Z = textures->u0 * chroma_u_scale;
-		consts->uv_xform2.W = textures->v0 * chroma_v_scale;
+		consts->uv_xform2.X = u0 * luma_u;
+		consts->uv_xform2.Y = v0 * luma_v;
+		consts->uv_xform2.Z = u0 * chroma_u;
+		consts->uv_xform2.W = v0 * chroma_v;
 	}
 
 	FBinkDrawVS::FParameters* vert_params = BinkGraphBuilder.AllocParameters<FBinkDrawVS::FParameters>();
@@ -450,13 +317,13 @@ static void Draw_textures(BINKTEXTURES* ptextures, BINKSHADERS* pshaders, void*
 
 	FBinkDrawICtCpPS::FPermutationDomain ictcp_pv;
 	ictcp_pv.Set<FBinkDrawICtCpPS::FALPHA>(!!hasAPlane);
-	ictcp_pv.Set<FBinkDrawICtCpPS::FTONEMAP>(textures->tonemap == 1);
-	ictcp_pv.Set<FBinkDrawICtCpPS::FST2084>(textures->tonemap == 2);
+	ictcp_pv.Set<FBinkDrawICtCpPS::FTONEMAP>(tonemap == 1);
+	ictcp_pv.Set<FBinkDrawICtCpPS::FST2084>(tonemap == 2);
 	TShaderMapRef<FBinkDrawICtCpPS> DrawICtCpPS(ShaderMap, ictcp_pv);
 
 	FBinkDrawYCbCrPS::FPermutationDomain ycbcr_pv;
 	ycbcr_pv.Set<FBinkDrawYCbCrPS::FALPHA>(!!hasAPlane);
-	ycbcr_pv.Set<FBinkDrawYCbCrPS::FSRGB>(!!(textures->draw_flags & 0x80000000));
+	ycbcr_pv.Set<FBinkDrawYCbCrPS::FSRGB>(!!(draw_flags & 0x80000000));
 	TShaderMapRef<FBinkDrawYCbCrPS> DrawYCbCrPS(ShaderMap, ycbcr_pv);
 
 	FRDGEventName EventName(TEXT("Bink"));
@@ -466,57 +333,49 @@ static void Draw_textures(BINKTEXTURES* ptextures, BINKSHADERS* pshaders, void*
 		ERDGPassFlags::Raster | ERDGPassFlags::SkipRenderPass,
 		[&](FRHICommandListImmediate& RHICmdList)
 		{
-			if (textures->dirty)
-			{
-				textures->dirty = 0;
-
-				FUpdateTextureRegion2D region_YAH(0, 0, 0, 0, bb->YABufferWidth, bb->YABufferHeight);
-				FUpdateTextureRegion2D region_cRcB(0, 0, 0, 0, bb->cRcBBufferWidth, bb->cRcBBufferHeight);
+			FUpdateTextureRegion2D region_YAH(0, 0, 0, 0, YABufferWidth, YABufferHeight);
+			FUpdateTextureRegion2D region_Ch(0, 0, 0, 0,  ChBufferWidth,  ChBufferHeight);
 
-				RHICmdList.UpdateTexture2D(textures->Ytexture[frame_num], 0, region_YAH, bp_src->YPlane.BufferPitch, (uint8*)bp_src->YPlane.Buffer);
-				RHICmdList.UpdateTexture2D(textures->cRtexture[frame_num], 0, region_cRcB, bp_src->cRPlane.BufferPitch, (uint8*)bp_src->cRPlane.Buffer);
-				RHICmdList.UpdateTexture2D(textures->cBtexture[frame_num], 0, region_cRcB, bp_src->cBPlane.BufferPitch, (uint8*)bp_src->cBPlane.Buffer);
-				if (hasAPlane) RHICmdList.UpdateTexture2D(textures->Atexture[frame_num], 0, region_YAH, bp_src->APlane.BufferPitch, (uint8*)bp_src->APlane.Buffer);
-				if (hasHPlane) RHICmdList.UpdateTexture2D(textures->Htexture[frame_num], 0, region_YAH, bp_src->HPlane.BufferPitch, (uint8*)bp_src->HPlane.Buffer);
+			RHICmdList.UpdateTexture2D(Ytexture[Idx],  0, region_YAH, YABufferWidth, (uint8*)Luma[Idx]);
+			RHICmdList.UpdateTexture2D(cRtexture[Idx], 0, region_Ch,  ChBufferWidth,  (uint8*)cR[Idx]);
+			RHICmdList.UpdateTexture2D(cBtexture[Idx], 0, region_Ch,  ChBufferWidth,  (uint8*)cB[Idx]);
+			if (hasAPlane && Alpha[Idx])
+			{
+				RHICmdList.UpdateTexture2D(Atexture[Idx], 0, region_YAH, YABufferWidth, (uint8*)Alpha[Idx]);
+			}
+			if (hasHPlane && HDR[Idx])
+			{
+				RHICmdList.UpdateTexture2D(Htexture[Idx], 0, region_YAH, YABufferWidth, (uint8*)HDR[Idx]);
 			}
-
 
 			RHICmdList.Transition(FRHITransitionInfo(consts->RenderTargets[0].GetTexture()->GetRHI()->GetTexture2D(), ERHIAccess::Unknown, ERHIAccess::RTV));
-			FRHIRenderPassInfo RPInfo(consts->RenderTargets[0].GetTexture()->GetRHI()->GetTexture2D(), ERenderTargetActions::Clear_Store);
+			FRHIRenderPassInfo RPInfo(consts->RenderTargets[0].GetTexture()->GetRHI()->GetTexture2D(), MakeRenderTargetActions(RenderTargetLoadAction, ERenderTargetStoreAction::EStore));
 			RHICmdList.BeginRenderPass(RPInfo, TEXT("BinkPass"));
 
 			FGraphicsPipelineStateInitializer GraphicsPSOInit;
 			RHICmdList.ApplyCachedRenderTargets(GraphicsPSOInit);
 
-			// Set up blending 
-			if ((!hasAPlane && textures->alpha >= 0.999f) || textures->draw_type == 2)
+			if ((!hasAPlane && alpha >= 0.999f) || draw_type == 2)
 			{
-				// opaque
 				GraphicsPSOInit.BlendState = TStaticBlendState<>::GetRHI();
 			}
-			else if (textures->draw_type == 1)
+			else if (draw_type == 1)
 			{
-				// alpha pre-multiplied
 				GraphicsPSOInit.BlendState = TStaticBlendState<CW_RGBA, BO_Add, BF_One, BF_InverseSourceAlpha, BO_Add, BF_One, BF_Zero>::GetRHI();
 			}
 			else
 			{
-				// normal alpha
 				GraphicsPSOInit.BlendState = TStaticBlendState<CW_RGBA, BO_Add, BF_SourceAlpha, BF_InverseSourceAlpha, BO_Add, BF_One, BF_Zero>::GetRHI();
 			}
 
-			// Disable backface culling
-			GraphicsPSOInit.RasterizerState = TStaticRasterizerState<FM_Solid, CM_None, ERasterizerDepthClipMode::DepthClip, false>::GetRHI();
-
-			// Disable Depth/Stencil test
+			GraphicsPSOInit.RasterizerState   = TStaticRasterizerState<FM_Solid, CM_None, ERasterizerDepthClipMode::DepthClip, false>::GetRHI();
 			GraphicsPSOInit.DepthStencilState = TStaticDepthStencilState<false, CF_Always>::GetRHI();
 
 			FVertexDeclarationElementList Elements;
 			GraphicsPSOInit.BoundShaderState.VertexDeclarationRHI = PipelineStateCache::GetOrCreateVertexDeclaration(Elements);
-			GraphicsPSOInit.BoundShaderState.VertexShaderRHI = DrawVS.GetVertexShader();
+			GraphicsPSOInit.BoundShaderState.VertexShaderRHI      = DrawVS.GetVertexShader();
 			GraphicsPSOInit.PrimitiveType = PT_TriangleStrip;
 
-			// Set VS/PS shaders
 			if (hasHPlane)
 			{
 				GraphicsPSOInit.BoundShaderState.PixelShaderRHI = DrawICtCpPS.GetPixelShader();
@@ -531,59 +390,49 @@ static void Draw_textures(BINKTEXTURES* ptextures, BINKSHADERS* pshaders, void*
 			}
 
 			SetShaderParameters(RHICmdList, DrawVS, DrawVS.GetVertexShader(), *vert_params);
-
 			RHICmdList.DrawPrimitive(0, 2, 1);
-
 			RHICmdList.EndRenderPass();
 		});
 
 	BinkGraphBuilder.Execute();
 }
 
-static void Set_draw_position(BINKTEXTURES * ptextures, F32 x0, F32 y0, F32 x1, F32 y1)
+void FBinkTextures::SetDrawPosition(float x0, float y0, float x1, float y1)
 {
-	BINKTEXTURESRHI * textures = (BINKTEXTURESRHI*)ptextures;
-	textures->Ax = x0;
-	textures->Ay = y0;
-	textures->Bx = x1;
-	textures->By = y0;
-	textures->Cx = x0;
-	textures->Cy = y1;
+	Ax = x0; Ay = y0;
+	Bx = x1; By = y0;
+	Cx = x0; Cy = y1;
 }
 
-static void Set_draw_corners(BINKTEXTURES * ptextures, F32 Ax, F32 Ay, F32 Bx, F32 By, F32 Cx, F32 Cy)
+void FBinkTextures::SetDrawCorners(float InAx, float InAy, float InBx, float InBy, float InCx, float InCy)
 {
-	BINKTEXTURESRHI * textures = (BINKTEXTURESRHI*)ptextures;
-	textures->Ax = Ax;
-	textures->Ay = Ay;
-	textures->Bx = Bx;
-	textures->By = By;
-	textures->Cx = Cx;
-	textures->Cy = Cy;
+	Ax = InAx; Ay = InAy;
+	Bx = InBx; By = InBy;
+	Cx = InCx; Cy = InCy;
 }
 
-static void Set_source_rect(BINKTEXTURES * ptextures, F32 u0, F32 v0, F32 u1, F32 v1)
+void FBinkTextures::SetSourceRect(float InU0, float InV0, float InU1, float InV1)
 {
-	BINKTEXTURESRHI * textures = (BINKTEXTURESRHI*)ptextures;
-	textures->u0 = u0;
-	textures->v0 = v0;
-	textures->u1 = u1;
-	textures->v1 = v1;
+	u0 = InU0; v0 = InV0;
+	u1 = InU1; v1 = InV1;
 }
 
-static void Set_alpha_settings(BINKTEXTURES * ptextures, F32 alpha_value, S32 draw_type)
+void FBinkTextures::SetAlphaSettings(float InAlpha, int32 InDrawFlags)
 {
-	BINKTEXTURESRHI * textures = (BINKTEXTURESRHI*)ptextures;
-	textures->alpha = alpha_value;
-	textures->draw_type = draw_type & 0x0FFFFFFF;
-	textures->draw_flags = draw_type & 0xF0000000;
+	alpha      = InAlpha;
+	draw_type  = InDrawFlags & 0x0FFFFFFF;
+	draw_flags = InDrawFlags & 0xF0000000;
 }
 
-static void Set_hdr_settings(BINKTEXTURES * ptextures, S32 tonemap, F32 exposure, S32 out_nits)
+void FBinkTextures::SetHDRSettings(int32 InTonemap, float InExposure, int32 InOutNits)
 {
-	BINKTEXTURESRHI * textures = (BINKTEXTURESRHI*)ptextures;
-	textures->tonemap = tonemap;
-	textures->exposure = exposure;
-	textures->out_luma = ((F32)out_nits) / 80.f;
+	tonemap  = InTonemap;
+	exposure = InExposure;
+	out_luma = ((F32)InOutNits) / 80.f;
 }
 
+void FBinkTextures::SetRenderTarget(FRHITexture* rt, ERenderTargetLoadAction load_action)
+{
+	RenderTarget = rt;
+	RenderTargetLoadAction = load_action;
+}
diff --git a/Engine/Plugins/Media/BinkMedia/Source/BinkMediaPlayer/Public/BinkMediaPlayer.h b/Engine/Plugins/Media/BinkMedia/Source/BinkMediaPlayer/Public/BinkMediaPlayer.h
--- a/Engine/Plugins/Media/BinkMedia/Source/BinkMediaPlayer/Public/BinkMediaPlayer.h
+++ b/Engine/Plugins/Media/BinkMedia/Source/BinkMediaPlayer/Public/BinkMediaPlayer.h
@@ -9,8 +9,9 @@
 #include "Slate/SlateTextures.h"
 #include "BinkMediaPlayer.generated.h"
 
-extern BINKMEDIAPLAYER_API unsigned bink_gpu_api;
-extern BINKMEDIAPLAYER_API unsigned bink_gpu_api_hdr;
+struct BINK;        // forward declare for HBINK (typedef struct BINK* HBINK in binktiny.h)
+struct FBinkTextures; // opaque Bink GPU texture state
+
 extern BINKMEDIAPLAYER_API EPixelFormat bink_force_pixel_format;
 extern BINKMEDIAPLAYER_API FString BinkUE4CookOnTheFlyPath(FString path, const TCHAR *filename);
 
@@ -473,11 +474,14 @@ public:
 	bool IsReady() const { return bnk != NULL; }
 	void UpdateTexture(const UTexture* tex, FRHICommandListImmediate &RHICmdList, FTextureRHIRef ref, void *nativePtr, int width, int height, bool isEditor, bool tonemap, int output_nits, float alpha, bool srgb_decode, bool is_hdr);
 
+
 	FIntPoint GetDimensions() const;
 	float GetFrameRate() const;
 
-	struct BINKPLUGIN *bnk;
+	struct BINK *bnk;              // HBINK handle (struct BINK* as defined in binktiny.h)
+	FBinkTextures *bink_textures; // render-thread GPU texture state (lazily created)
 	bool paused;
 	bool reached_end;
 	FDelegateHandle overlayHook;
+
 };
diff --git a/Engine/Plugins/Media/BinkMedia/Source/BinkMediaPlayer/Public/BinkUE.h b/Engine/Plugins/Media/BinkMedia/Source/BinkMediaPlayer/Public/BinkUE.h
new file mode 100644
--- /dev/null
+++ b/Engine/Plugins/Media/BinkMedia/Source/BinkMediaPlayer/Public/BinkUE.h
@@ -0,0 +1,23 @@
+// Copyright Epic Games Tools LLC
+//   Licenced under the Unreal Engine EULA 
+
+#ifndef BINKPLUGIN_UE4_H
+#define BINKPLUGIN_UE4_H
+
+// you can define these variables at the project compiler define level
+//   if you want the Bink data or binaries to live somewhere else
+
+#ifndef BINKCONTENTPATH
+#define BINKCONTENTPATH FPaths::ProjectContentDir()
+#endif
+
+#ifndef BINKMOVIEPATH
+#define BINKMOVIEPATH BINKCONTENTPATH + TEXT("Movies/")
+#endif
+
+#ifndef BINKTEMPPATH
+#define BINKTEMPPATH FPaths::ProjectSavedDir() + TEXT("Bink/")
+#endif
+
+#endif // BINKPLUGIN_UE4_H
+
diff --git a/Engine/Plugins/Media/BinkMedia/Source/BinkMediaPlayerEditor/Private/BinkMediaPlayerCustomization.cpp b/Engine/Plugins/Media/BinkMedia/Source/BinkMediaPlayerEditor/Private/BinkMediaPlayerCustomization.cpp
--- a/Engine/Plugins/Media/BinkMedia/Source/BinkMediaPlayerEditor/Private/BinkMediaPlayerCustomization.cpp
+++ b/Engine/Plugins/Media/BinkMedia/Source/BinkMediaPlayerEditor/Private/BinkMediaPlayerCustomization.cpp
@@ -16,7 +16,7 @@
 #include "Widgets/Input/SButton.h"
 #include "Widgets/Input/SEditableText.h"
 #include "Widgets/SWindow.h"
-#include "binkplugin_ue4.h"
+#include "BinkUE.h"
 
 #define LOCTEXT_NAMESPACE "FBinkMediaPlayerCustomization"
 
diff --git a/Engine/Plugins/Media/BinkMedia/Source/BinkMediaPlayerEditor/Private/Factories/BinkMediaPlayerFactory.cpp b/Engine/Plugins/Media/BinkMedia/Source/BinkMediaPlayerEditor/Private/Factories/BinkMediaPlayerFactory.cpp
--- a/Engine/Plugins/Media/BinkMedia/Source/BinkMediaPlayerEditor/Private/Factories/BinkMediaPlayerFactory.cpp
+++ b/Engine/Plugins/Media/BinkMedia/Source/BinkMediaPlayerEditor/Private/Factories/BinkMediaPlayerFactory.cpp
@@ -5,7 +5,7 @@
 
 #include "Misc/Paths.h"
 #include "BinkMediaPlayer.h"
-#include "binkplugin_ue4.h"
+#include "BinkUE.h"
 
 UBinkMediaPlayerFactory::UBinkMediaPlayerFactory( const FObjectInitializer& ObjectInitializer )
 	: Super(ObjectInitializer)
diff --git a/Engine/Plugins/Media/BinkMedia/Source/BinkMediaPlayerEditor/Private/SBinkWidgets.cpp b/Engine/Plugins/Media/BinkMedia/Source/BinkMediaPlayerEditor/Private/SBinkWidgets.cpp
--- a/Engine/Plugins/Media/BinkMedia/Source/BinkMediaPlayerEditor/Private/SBinkWidgets.cpp
+++ b/Engine/Plugins/Media/BinkMedia/Source/BinkMediaPlayerEditor/Private/SBinkWidgets.cpp
@@ -41,7 +41,7 @@ struct FBinkMediaPlayerEditorViewport : public ISlateViewport
 		else 
 		{
 			const FIntPoint dim = MediaPlayer->GetDimensions();
-			EPixelFormat format = bink_force_pixel_format != PF_Unknown ? bink_force_pixel_format : bink_gpu_api_hdr ? PF_A2B10G10R10 : PF_B8G8R8A8;
+			EPixelFormat format = bink_force_pixel_format != PF_Unknown ? bink_force_pixel_format : GRHISupportsHDROutput ? PF_A2B10G10R10 : PF_B8G8R8A8;
 			SlateTexture = new FSlateTexture2DRHIRef(dim.X, dim.Y, format, nullptr, TexCreate_RenderTargetable, true);
 			EditorTexture = new FBinkMediaPlayerEditorTexture(SlateTexture, MediaPlayer);
 		}
diff --git a/Engine/Plugins/Media/BinkMedia/Source/BinkMediaPlayerSDK/include/binkplugin.h b/Engine/Plugins/Media/BinkMedia/Source/BinkMediaPlayerSDK/include/binkplugin.h
deleted file mode 100644
diff --git a/Engine/Plugins/Media/BinkMedia/Source/BinkMediaPlayerSDK/include/binktextures.h b/Engine/Plugins/Media/BinkMedia/Source/BinkMediaPlayerSDK/include/binktextures.h
deleted file mode 100644
diff --git a/Engine/Plugins/Media/BinkMedia/Source/BinkMediaPlayerSDK/include/binktiny.h b/Engine/Plugins/Media/BinkMedia/Source/BinkMediaPlayerSDK/include/binktiny.h
--- a/Engine/Plugins/Media/BinkMedia/Source/BinkMediaPlayerSDK/include/binktiny.h
+++ b/Engine/Plugins/Media/BinkMedia/Source/BinkMediaPlayerSDK/include/binktiny.h
@@ -2,55 +2,15 @@
 // Licensed under the terms of a valid Unreal Engine license agreement,
 //   and the separate 'Unreal Engine End User License Agreement for Publishing'.
 
-// Copyright Epic Games Tools, LLC. All Rights Reserved.
+// Copyright Epic Games, Inc. All Rights Reserved.
 
 #ifndef __BINKH__
 #define __BINKH__
 
 typedef struct BINK * HBINK;
 
-#define BINKALPHA              0x00100000L // Decompress alpha plane (if present)
-#define BINKHDR                0x00000004L // Video is an HDR video
-#define BINKUSETRIPLEBUFFERING 0x00000008L // Use triple buffering in the framebuffers
-
-typedef struct BINKRECT {
-  S32 Left,Top,Width,Height;
-} BINKRECT;
-
-#define BINKMAXDIRTYRECTS 1
-
-typedef struct BINKPLANE
-{
-  void * Buffer;
-  S32 Allocate;
-  U32 BufferPitch;
-} BINKPLANE;
-
-typedef struct BINKFRAMEPLANESET
+typedef struct BINK 
 {
-  BINKPLANE YPlane;
-  BINKPLANE cRPlane;
-  BINKPLANE cBPlane;
-  BINKPLANE APlane;
-  BINKPLANE HPlane;
-} BINKFRAMEPLANESET;
-
-#define BINKMAXFRAMEBUFFERS 3
-
-typedef struct BINKFRAMEBUFFERS
-{
-  S32 TotalFrames;
-  U32 YABufferWidth;
-  U32 YABufferHeight;
-  U32 cRcBBufferWidth;
-  U32 cRcBBufferHeight;
-
-  U32 FrameNum;  // 0 to (TotalFrames-1)
-  BINKFRAMEPLANESET Frames[ BINKMAXFRAMEBUFFERS ];
-} BINKFRAMEBUFFERS;
-
-
-typedef struct BINK {
   U32 Width;                  // Width (1 based, 640 for example)
   U32 Height;                 // Height (1 based, 480 for example)
   U32 Frames;                 // Number of frames (1 based, 100 = 100 frames)
@@ -69,27 +29,109 @@ typedef struct BINK {
   U32 SndSize;                // The current frame sound tracks' size in bytes
   U32 FrameChangePercent;     // very rough percentage of the frame that changed
   S32 NumTracks;              // how many tracks
-  S32 NumRects;
 
-  U32 soundon;                // sound turned on?
-  U32 videoon;                // video turned on?
-  U32 needio;                 // set to 1, if we need an io before the next binkframe
-  S32 closing;                // are we closing?
+} BINK;
 
-  BINKRECT FrameRects[BINKMAXDIRTYRECTS];// Dirty rects from BinkGetRects
+// ================================================================================
+// Bink API to open and search for offsets within a zip file (play bink in zip if bink is uncompressed)
 
-  F32 ColorSpace[16];         // 4x4 matrix to use for colorspace conversion
-                              // in HDR, 0 = smpte2084(maxnits/10k)*0.5, 1/2=ct/cp scale, 3/4=ct/cp bias, 5 = maxnits
+RADEXPFUNC void * RADEXPLINK BinkZipOpen( char const * zipfn );
+RADEXPFUNC S32 RADEXPLINK BinkZipFindUncompressedFileOffset( void * zip, char const * filename, U64 * out_ofs );
+RADEXPFUNC S32 RADEXPLINK BinkZipGetFileInfo( void * zip, U64 index, char * out_filename, U32 bytes_filename, U64 * out_offset, S32 * out_uncompressed );
+RADEXPFUNC void RADEXPLINK BinkZipClose( void * zip );
 
-  U64 FileOffset;             // Offset into the file where the Bink starts        
-  U64 Highest1SecRate;        // Highest 1 sec data rate
 
-  BINKFRAMEBUFFERS * FrameBuffers; // Bink frame buffers that we decompress to
+// ================================================================================
+// Bink high level API (usually used by Unity)
 
-} BINK;
+RADEXPFUNC void RADEXPLINK BinkHLInit( U32 thread_count, U32 snd_freq, U32 snd_chans ); // snd_freq == 0, no sound init
+RADEXPFUNC void RADEXPLINK BinkHLShutdown( void );
+
+#define BINKHLSNDLAYOUT_NONE               0 // don't open any sound tracks snd_track_start not used
+#define BINKHLSNDLAYOUT_SIMPLE             1 // based on filename, OR simply mono or stereo sound in track snd_track_start (default speaker spread)
+#define BINKHLSNDLAYOUT_LANGUAGEOVERRIDE   2 // mono or stereo sound in track 0, language track at snd_track_start
+#define BINKHLSNDLAYOUT_51                 3 // 6 mono tracks in tracks snd_track_start[0..5]
+#define BINKHLSNDLAYOUT_51LANGUAGEOVERRIDE 4 // 6 mono tracks in tracks 0..5, center language track at snd_track_start
+#define BINKHLSNDLAYOUT_71                 5 // 8 mono tracks in tracks snd_track_start[0..7]
+#define BINKHLSNDLAYOUT_71LANGUAGEOVERRIDE 6 // 8 mono tracks in tracks 0..7, center language track at snd_track_start
 
-RADEXPFUNC void RADEXPLINK BinkGetFrameBuffersInfo( HBINK bink, BINKFRAMEBUFFERS * fbset );
-RADEXPFUNC S32  RADEXPLINK BinkAllocateFrameBuffers( HBINK bp, BINKFRAMEBUFFERS * set, U32 minimum_alignment );
-RADEXPFUNC void RADEXPLINK BinkRegisterFrameBuffers( HBINK bink, BINKFRAMEBUFFERS * fbset );
+#define BINKHLBUFFER_STREAM              0 // stream the movie off the media during playback (caches about 1 second of video)
+#define BINKHLBUFFER_LOADALL             1 // loads the whole movie into memory at Open time (will block)
+#define BINKHLBUFFER_STREAMUNTILRESIDENT 2 // streams the movie into a memory buffer as big as the movie, so it will be preloaded eventually)
+
+
+RADEXPFUNC HBINK RADEXPLINK BinkHLOpen( char const * name, U32 snd_track_type, U32 snd_track_start, U32 buffering_type, U64 file_byte_offset );
+
+typedef struct BINKHLPRELOADED BINKHLPRELOADED;
+RADEXPFUNC BINKHLPRELOADED * RADEXPLINK BinkHLPreload( char const * name, U64 file_byte_offset );
+RADEXPFUNC void RADEXPLINK BinkHLUnpreload( BINKHLPRELOADED * pre );
+RADEXPFUNC HBINK RADEXPLINK BinkHLOpenPreload( BINKHLPRELOADED * pre, U32 snd_track_type, U32 snd_track_start );
+
+typedef struct BINKHLINFO
+{
+  U32 Width;
+  U32 Height;
+  U32 TextureBufferYAWidth;
+  U32 TextureBufferYAHeight;
+  U32 TextureBufferChWidth;
+  U32 TextureBufferChHeight;
+  U32 TextureDrawIndex;
+  U32 TextureDrawFrame;     // what frame would be drawn with TextureDrawIndex
+  S32 NeedAlpha;
+  S32 NeedHDR;
+  U32 Frames;
+  U32 FrameNum;
+  U32 FrameRate;
+  U32 FrameRateDiv;
+  U32 LoopsRemaining;
+  U32 ReadError;
+  U32 NumTracksOpened;
+  U32 PlaybackState;         // 0 = playing, 1 = paused, 2 = gotoing, 3 = at end (stopped)
+  F32 cs0, cs1, cs2, cs3, cs4, cs5, cs6, cs7, cs8, cs9, cs10, cs11, cs12, cs13, cs14, cs15;
+} BINKHLINFO;
+
+RADEXPFUNC void RADEXPLINK BinkHLGetInfo( HBINK b, BINKHLINFO * info );
+
+RADEXPFUNC void RADEXPLINK BinkHLRegisterTextureBuffers( HBINK b, U8 ** Luma,  // two luma pointers
+                                                                  U8 ** cR,    // two chroma pointers
+                                                                  U8 ** cB, 
+                                                                  U8 ** Alpha,
+                                                                  U8 ** HDR );
+
+RADEXPFUNC void RADEXPLINK BinkHLPause( HBINK b, U32 on_off );
+
+RADEXPFUNC void RADEXPLINK BinkHLVolume( HBINK b, F32 vol );
+
+// set speaker volumes
+//   BinkSndSimple = count must be 2 (l/r) - OR match filename (is there a "_51." or "_51L." or "_71." or "_71L." in the filename?)
+//   BinkSndLanguageOverride = count must be 3 (l,r)/language
+//   BinkSnd51 = count must be 6 (front l/r),center,sub,(rear l/r),
+//   BinkSnd51LanguageOverride = 7 (front l/r),center,sub,(rear l/r),language
+//   BinkSnd71 = count must be 8 (front l/r),center,sub,(read l/r),(surr l/r)
+//   BinkSnd71LanguageOverride = 9 (front l/r),center,sub,(read l/r),(surr l/r), lang
+RADEXPFUNC void RADEXPLINK BinkHLSpeakerVolumes( HBINK b, F32 * vols, U32 count );
+
+RADEXPFUNC void RADEXPLINK BinkHLLoopCount( HBINK b, U32 count );  // 0 = loop forever
+
+RADEXPFUNC void RADEXPLINK BinkHLGoto( HBINK b, U32 goto_frame, U32 max_frames_to_decode_per_process );  // goto frame goes from 1 to framenum (0 means stop current goto), max_frames_to_decode_per_process = 0 means block until at frame
+
+RADEXPFUNC void RADEXPLINK BinkHLProcess( void ); // decodes/moves to next frame (next frame, or loop, or goto) for all Binks that are playing (if they need advancing)
+
+RADEXPFUNC char * RADEXPLINK BinkGetError( void );
+
+RADEXPFUNC void RADEXPLINK BinkClose( HBINK bnk );
+
+#define BINKSURFACE32BGRx      3
+#define BINKSURFACE32RGBx      4
+#define BINKSURFACE32BGRA      5
+#define BINKSURFACE32RGBA      6
+#define BINKSURFACE32ARGB     12
+
+RADEXPFUNC S32  RADEXPLINK BinkCopyToBuffer(HBINK bnk,void* dest,S32 destpitch,U32 destheight,U32 destx,U32 desty,U32 flags);
+RADEXPFUNC S32  RADEXPLINK BinkCopyToBufferRect(HBINK bnk,void* dest,S32 destpitch,U32 destheight,U32 destx,U32 desty,U32 srcx, U32 srcy, U32 srcw, U32 srch, U32 flags);
+
+RADEXPFUNC S32 RADEXPLINK BinkLoadSubtitles( HBINK bink, char const * srt_file );
+RADEXPFUNC char const * RADEXPLINK BinkCurrentSubtitle( HBINK bink, U32 * iterate, U32 * opt_out_start_ms, U32 * opt_out_end_ms );
+RADEXPFUNC char const * RADEXPLINK BinkGetSubtitleByIndex( HBINK bink, U32 index, U32 * opt_out_start_ms, U32 * opt_out_end_ms );
 
 #endif
diff --git a/Engine/Plugins/Media/BinkMedia/Source/BinkMediaPlayerSDK/include/egttypes.h b/Engine/Plugins/Media/BinkMedia/Source/BinkMediaPlayerSDK/include/egttypes.h
--- a/Engine/Plugins/Media/BinkMedia/Source/BinkMediaPlayerSDK/include/egttypes.h
+++ b/Engine/Plugins/Media/BinkMedia/Source/BinkMediaPlayerSDK/include/egttypes.h
@@ -3,7 +3,7 @@
 //   and the separate 'Unreal Engine End User License Agreement for Publishing'.
 
 //===============================================
-// Copyright Epic Games Tools, LLC. All Rights Reserved.
+// Copyright Epic Games All Rights Reserved.
 //===============================================
 
 #ifndef __EGTTYPESH__
@@ -13,7 +13,7 @@
 #define __RADTYPESH__
 #define __RADRR_COREH__ // block old rr_core
 
-#define RADCOPYRIGHT "Copyright (C) 1994-2025, Epic Games Tools, LLC"
+#define RADCOPYRIGHT "Copyright (C) 1994-2026, Epic Games, Inc."
 
 #if !defined(__RADRES__) // don't include anything for resource compiles
 
