Bink: Fixed foward def mismatch, just include PixelFormat.h

#Synced-CL 52144611

diff --git a/Engine/Plugins/Media/BinkMedia/Source/BinkMediaPlayer/Private/BinkMediaPlayerPrivate.h b/Engine/Plugins/Media/BinkMedia/Source/BinkMediaPlayer/Private/BinkMediaPlayerPrivate.h
--- a/Engine/Plugins/Media/BinkMedia/Source/BinkMediaPlayer/Private/BinkMediaPlayerPrivate.h
+++ b/Engine/Plugins/Media/BinkMedia/Source/BinkMediaPlayer/Private/BinkMediaPlayerPrivate.h
@@ -6,6 +6,7 @@
 #include "Rendering/SlateRenderer.h"
 #include "AudioMixerDevice.h"
 #include "RHIFwd.h"
+#include "PixelFormat.h"
 
 #if PLATFORM_ANDROID
 #include <android/log.h>
@@ -19,8 +20,6 @@
 #include "binktiny.h"
 #include "BinkRHI.h"
 
-enum EPixelFormat : uint8;
-
 extern BINKMEDIAPLAYER_API EPixelFormat bink_force_pixel_format;
 extern BINKMEDIAPLAYER_API FString BinkUE4CookOnTheFlyPath(FString path, const TCHAR *filename);
 extern BINKMEDIAPLAYER_API bool BinkInitialize();
