Bink: Fixed BinkTestBed

#Synced-CL 52362648

diff --git a/Engine/Plugins/Media/BinkMedia/Source/BinkMediaPlayer/Private/BinkRHI.h b/Engine/Plugins/Media/BinkMedia/Source/BinkMediaPlayer/Private/BinkRHI.h
--- a/Engine/Plugins/Media/BinkMedia/Source/BinkMediaPlayer/Private/BinkRHI.h
+++ b/Engine/Plugins/Media/BinkMedia/Source/BinkMediaPlayer/Private/BinkRHI.h
@@ -16,27 +16,27 @@ struct FBinkTextures
 	// --- Lifecycle -----------------------------------------------------------
 
 	// Full creation (CPU + GPU). Must be called on the render thread.
-	static FBinkTextures* Create(HBINK bink, void* user_ptr = nullptr);
+	BINKMEDIAPLAYER_API static FBinkTextures* Create(HBINK bink, void* user_ptr = nullptr);
 
 	// Creates only the CPU-side decode buffers and registers them with BinkHL.
 	// Safe to call from the game thread. GPU textures are created lazily on
 	// the render thread the first time Draw() is called.
-	static FBinkTextures* CreateCPU(HBINK bink, void* user_ptr = nullptr);
+	BINKMEDIAPLAYER_API static FBinkTextures* CreateCPU(HBINK bink, void* user_ptr = nullptr);
 
 	// Frees all GPU + CPU resources.
-	void Destroy();
+	BINKMEDIAPLAYER_API void Destroy();
 
 	// --- Per-frame draw setup (call before Draw) -----------------------------
 
-	void SetRenderTarget(FRHITexture* rt, ERenderTargetLoadAction load_action);
-	void SetDrawPosition(float x0, float y0, float x1, float y1);
-	void SetDrawCorners(float Ax, float Ay, float Bx, float By, float Cx, float Cy);
-	void SetSourceRect(float u0, float v0, float u1, float v1);
-	void SetAlphaSettings(float alpha, int32 draw_flags);
-	void SetHDRSettings(int32 tonemap, float exposure, int32 out_nits);
+	BINKMEDIAPLAYER_API void SetRenderTarget(FRHITexture* rt, ERenderTargetLoadAction load_action);
+	BINKMEDIAPLAYER_API void SetDrawPosition(float x0, float y0, float x1, float y1);
+	BINKMEDIAPLAYER_API void SetDrawCorners(float Ax, float Ay, float Bx, float By, float Cx, float Cy);
+	BINKMEDIAPLAYER_API void SetSourceRect(float u0, float v0, float u1, float v1);
+	BINKMEDIAPLAYER_API void SetAlphaSettings(float alpha, int32 draw_flags);
+	BINKMEDIAPLAYER_API void SetHDRSettings(int32 tonemap, float exposure, int32 out_nits);
 
 	// Draws the current Bink frame to the previously-set render target.
-	void Draw();
+	BINKMEDIAPLAYER_API void Draw();
 
 	// --- Internal state (treat as private) -----------------------------------
 
diff --git a/Engine/Plugins/Media/BinkMedia/Source/BinkMediaPlayer/Private/binktexturesRHI.cpp b/Engine/Plugins/Media/BinkMedia/Source/BinkMediaPlayer/Private/binktexturesRHI.cpp
--- a/Engine/Plugins/Media/BinkMedia/Source/BinkMediaPlayer/Private/binktexturesRHI.cpp
+++ b/Engine/Plugins/Media/BinkMedia/Source/BinkMediaPlayer/Private/binktexturesRHI.cpp
@@ -23,7 +23,7 @@ FRDGTextureRef BinkRegisterExternalTexture(FRDGBuilder& GraphBuilder, FRHITextur
 // Struct definition is in BinkRHI.h
 
 //-----------------------------------------------------------------------------
-FBinkTextures* FBinkTextures::Create(HBINK bink, void* user_ptr)
+BINKMEDIAPLAYER_API FBinkTextures* FBinkTextures::Create(HBINK bink, void* user_ptr)
 {
 	FBinkTextures* textures = FBinkTextures::CreateCPU(bink, user_ptr);
 	if (!textures)
@@ -38,7 +38,7 @@ FBinkTextures* FBinkTextures::Create(HBINK bink, void* user_ptr)
 // Creates CPU-side decode buffers and registers them with BinkHL.
 // Safe to call from the game thread. GPU textures are created lazily
 // on the render thread the first time Draw_textures is called.
-FBinkTextures* FBinkTextures::CreateCPU(HBINK bink, void* user_ptr)
+BINKMEDIAPLAYER_API FBinkTextures* FBinkTextures::CreateCPU(HBINK bink, void* user_ptr)
 {
 	FBinkTextures* textures = (FBinkTextures*)FMemory::Malloc(sizeof(*textures));
 	if (textures == 0)
@@ -146,7 +146,7 @@ void FBinkTextures::EnsureGpuTextures()
 
 
 //-----------------------------------------------------------------------------
-void FBinkTextures::Destroy()
+BINKMEDIAPLAYER_API void FBinkTextures::Destroy()
 {
 
 	// Clear the flag BEFORE releasing GPU resources so that any concurrent reader
@@ -173,7 +173,7 @@ void FBinkTextures::Destroy()
 }
 
 //-----------------------------------------------------------------------------
-void FBinkTextures::Draw()
+BINKMEDIAPLAYER_API void FBinkTextures::Draw()
 {
 
 	if (RenderTarget == nullptr)
@@ -397,41 +397,41 @@ void FBinkTextures::Draw()
 	BinkGraphBuilder.Execute();
 }
 
-void FBinkTextures::SetDrawPosition(float x0, float y0, float x1, float y1)
+BINKMEDIAPLAYER_API void FBinkTextures::SetDrawPosition(float x0, float y0, float x1, float y1)
 {
 	Ax = x0; Ay = y0;
 	Bx = x1; By = y0;
 	Cx = x0; Cy = y1;
 }
 
-void FBinkTextures::SetDrawCorners(float InAx, float InAy, float InBx, float InBy, float InCx, float InCy)
+BINKMEDIAPLAYER_API void FBinkTextures::SetDrawCorners(float InAx, float InAy, float InBx, float InBy, float InCx, float InCy)
 {
 	Ax = InAx; Ay = InAy;
 	Bx = InBx; By = InBy;
 	Cx = InCx; Cy = InCy;
 }
 
-void FBinkTextures::SetSourceRect(float InU0, float InV0, float InU1, float InV1)
+BINKMEDIAPLAYER_API void FBinkTextures::SetSourceRect(float InU0, float InV0, float InU1, float InV1)
 {
 	u0 = InU0; v0 = InV0;
 	u1 = InU1; v1 = InV1;
 }
 
-void FBinkTextures::SetAlphaSettings(float InAlpha, int32 InDrawFlags)
+BINKMEDIAPLAYER_API void FBinkTextures::SetAlphaSettings(float InAlpha, int32 InDrawFlags)
 {
 	alpha      = InAlpha;
 	draw_type  = InDrawFlags & 0x0FFFFFFF;
 	draw_flags = InDrawFlags & 0xF0000000;
 }
 
-void FBinkTextures::SetHDRSettings(int32 InTonemap, float InExposure, int32 InOutNits)
+BINKMEDIAPLAYER_API void FBinkTextures::SetHDRSettings(int32 InTonemap, float InExposure, int32 InOutNits)
 {
 	tonemap  = InTonemap;
 	exposure = InExposure;
 	out_luma = ((F32)InOutNits) / 80.f;
 }
 
-void FBinkTextures::SetRenderTarget(FRHITexture* rt, ERenderTargetLoadAction load_action)
+BINKMEDIAPLAYER_API void FBinkTextures::SetRenderTarget(FRHITexture* rt, ERenderTargetLoadAction load_action)
 {
 	RenderTarget = rt;
 	RenderTargetLoadAction = load_action;
